First pass of OneSky Localization Service Plugin
- Implements the OneSky Project Group, Project, File, and Translation APIs, minus delete commands, etc.
- The communication with OneSky works, but the interface between the Localization Service module and the OneSky plugin about what general Localization Service Commands map to which OneSky commands needs to be implemented.
- Some debug code remains, and some clean up still needs to be done
- Some minor changes to Localization Service module

diff --git a/Engine/Plugins/Developer/OneSkyLocalizationService/OneSkyLocalizationService.uplugin b/Engine/Plugins/Developer/OneSkyLocalizationService/OneSkyLocalizationService.uplugin
new file mode 100644
--- /dev/null
+++ b/Engine/Plugins/Developer/OneSkyLocalizationService/OneSkyLocalizationService.uplugin
@@ -0,0 +1,21 @@
+{
+	"FileVersion" : 3,
+	
+	"FriendlyName" : "OneSky",
+	"Version" : 1,
+	"VersionName" : "1.0",
+	"CreatedBy" : "Epic Games, Inc.",
+	"CreatedByURL" : "http://epicgames.com",
+	"EngineVersion" : "4.8.0",
+	"Description" : "OneSky localization service",
+	"Category" : "Editor.Localization Service",
+	"EnabledByDefault" : false,
+	
+	"Modules" :
+	[
+		{
+			"Name" : "OneSkyLocalizationService",
+			"Type" : "Developer"
+		}
+	]
+}
\ No newline at end of file
diff --git a/Engine/Plugins/Developer/OneSkyLocalizationService/Source/OneSkyLocalizationService/OneSkyLocalizationService.Build.cs b/Engine/Plugins/Developer/OneSkyLocalizationService/Source/OneSkyLocalizationService/OneSkyLocalizationService.Build.cs
new file mode 100644
--- /dev/null
+++ b/Engine/Plugins/Developer/OneSkyLocalizationService/Source/OneSkyLocalizationService/OneSkyLocalizationService.Build.cs
@@ -0,0 +1,34 @@
+// Copyright 1998-2015 Epic Games, Inc. All Rights Reserved.
+
+using UnrealBuildTool;
+
+public class OneSkyLocalizationService : ModuleRules
+{
+    public OneSkyLocalizationService(TargetInfo Target)
+	{
+        PrivateDependencyModuleNames.AddRange(
+            new string[] {
+				"Core",
+				"CoreUObject",
+                "InputCore",
+				"Slate",
+				"SlateCore",
+                "EditorStyle",
+                "LocalizationService",
+                "Json",
+                "Internationalization",
+                "HTTP",
+                "Serialization",
+			}
+		);
+
+        PrivateIncludePathModuleNames.AddRange(
+            new string[] {
+				"LocalizationService",
+				"Json",
+                "Internationalization",
+				"HTTP",
+			}
+        );
+	}
+}
diff --git a/Engine/Plugins/Developer/OneSkyLocalizationService/Source/OneSkyLocalizationService/Private/IOneSkyLocalizationServiceWorker.h b/Engine/Plugins/Developer/OneSkyLocalizationService/Source/OneSkyLocalizationService/Private/IOneSkyLocalizationServiceWorker.h
new file mode 100644
--- /dev/null
+++ b/Engine/Plugins/Developer/OneSkyLocalizationService/Source/OneSkyLocalizationService/Private/IOneSkyLocalizationServiceWorker.h
@@ -0,0 +1,26 @@
+// Copyright 1998-2015 Epic Games, Inc. All Rights Reserved.
+
+#pragma once
+
+class IOneSkyLocalizationServiceWorker
+{
+public:
+	/**
+	 * Name describing the work that this worker does. Used for factory method hookup.
+	 */
+	virtual FName GetName() const = 0;
+
+	/**
+	 * Function that actually does the work. Can be executed on another thread.
+	 */
+	virtual bool Execute( class FOneSkyLocalizationServiceCommand& InCommand ) = 0;
+
+	/**
+	 * Updates the state of any items after completion (if necessary). This is always executed on the main thread.
+	 * @returns true if states were updated
+	 */
+	virtual bool UpdateStates() const = 0;
+
+};
+
+typedef TSharedRef<IOneSkyLocalizationServiceWorker, ESPMode::ThreadSafe> FOneSkyLocalizationServiceWorkerRef;
\ No newline at end of file
diff --git a/Engine/Plugins/Developer/OneSkyLocalizationService/Source/OneSkyLocalizationService/Private/OneSkyConnection.cpp b/Engine/Plugins/Developer/OneSkyLocalizationService/Source/OneSkyLocalizationService/Private/OneSkyConnection.cpp
new file mode 100644
--- /dev/null
+++ b/Engine/Plugins/Developer/OneSkyLocalizationService/Source/OneSkyLocalizationService/Private/OneSkyConnection.cpp
@@ -0,0 +1,96 @@
+// Copyright 1998-2015 Epic Games, Inc. All Rights Reserved.
+
+#include "OneSkyLocalizationServicePrivatePCH.h"
+#include "OneSkyLocalizationServiceModule.h"
+#include "OneSkyConnection.h"
+#include "OneSkyLocalizationServiceCommand.h"
+#include "OneSkyLocalizationServiceProvider.h"
+#include "MessageLog.h"
+
+#define LOCTEXT_NAMESPACE "OneSkyConnection"
+
+FOneSkyConnection::FOneSkyConnection(const FOneSkyConnectionInfo& InConnectionInfo)
+:	bEstablishedConnection(false)
+{
+	EstablishConnection(InConnectionInfo);
+}
+
+FOneSkyConnection::~FOneSkyConnection()
+{
+	Disconnect();
+}
+
+bool FOneSkyConnection::IsValidConnection()
+{
+	return bEstablishedConnection;
+}
+
+void FOneSkyConnection::Disconnect()
+{
+	// TODO: Maybe do something
+}
+
+bool FOneSkyConnection::EstablishConnection(const FOneSkyConnectionInfo& InConnectionInfo)
+{
+	//Connection assumed successful
+	bEstablishedConnection = true;
+
+	UE_LOG(LogLocalizationService, Verbose, TEXT("OneSky connection created: %s"), *InConnectionInfo.Name);
+	
+	return bEstablishedConnection;
+}
+
+FScopedOneSkyConnection::FScopedOneSkyConnection( const class FOneSkyLocalizationServiceCommand& InCommand )
+	: Connection(NULL)
+	, Concurrency(InCommand.Concurrency)
+{
+	Initialize(InCommand.ConnectionInfo);
+}
+
+FScopedOneSkyConnection::FScopedOneSkyConnection( ELocalizationServiceOperationConcurrency::Type InConcurrency, const FOneSkyConnectionInfo& InConnectionInfo )
+	: Connection(NULL)
+	, Concurrency(InConcurrency)
+{
+	Initialize(InConnectionInfo);
+}
+
+void FScopedOneSkyConnection::Initialize( const FOneSkyConnectionInfo& InConnectionInfo )
+{
+	if (Concurrency == ELocalizationServiceOperationConcurrency::Synchronous)
+	{
+		// Synchronous commands reuse the same persistent connection to reduce
+		// the number of expensive connection attempts.
+		FOneSkyLocalizationServiceModule& OneSkyLocalizationService = FModuleManager::LoadModuleChecked<FOneSkyLocalizationServiceModule>("OneSkyLocalizationService");
+		if ( OneSkyLocalizationService.GetProvider().EstablishPersistentConnection() )
+		{
+			Connection = OneSkyLocalizationService.GetProvider().GetPersistentConnection();
+		}
+	}
+	else
+	{
+		// Async commands form a new connection for each attempt because
+		// using the persistent connection is not threadsafe
+		FOneSkyConnection* NewConnection = new FOneSkyConnection(InConnectionInfo);
+		if ( NewConnection->IsValidConnection() )
+		{
+			Connection = NewConnection;
+		}
+	}
+}
+
+FScopedOneSkyConnection::~FScopedOneSkyConnection()
+{
+	if(Connection != NULL)
+	{
+		if (Concurrency == ELocalizationServiceOperationConcurrency::Asynchronous)
+		{
+			// Remove this connection as it is only temporary
+			Connection->Disconnect();
+			delete Connection;
+		}
+		
+		Connection = NULL;
+	}
+}
+
+#undef LOCTEXT_NAMESPACE
diff --git a/Engine/Plugins/Developer/OneSkyLocalizationService/Source/OneSkyLocalizationService/Private/OneSkyConnection.h b/Engine/Plugins/Developer/OneSkyLocalizationService/Source/OneSkyLocalizationService/Private/OneSkyConnection.h
new file mode 100644
--- /dev/null
+++ b/Engine/Plugins/Developer/OneSkyLocalizationService/Source/OneSkyLocalizationService/Private/OneSkyConnection.h
@@ -0,0 +1,84 @@
+// Copyright 1998-2015 Epic Games, Inc. All Rights Reserved.
+
+#pragma once
+
+#include "Developer/LocalizationService/Public/ILocalizationServiceProvider.h"
+#include "OneSkyLocalizationServiceCommand.h"
+
+class FOneSkyConnection
+{
+public:
+	
+
+	FOneSkyConnection(const FOneSkyConnectionInfo& InConnectionInfo);
+	/** API Specific close of localization service project*/
+	~FOneSkyConnection();
+
+
+	/** Returns true if connection is currently active */
+	bool IsValidConnection();
+
+	/** If the connection is valid, disconnect from the server */
+	void Disconnect();
+
+	/**
+	 * Make a valid connection if possible
+	 * @param InConnectionInfo		Connection credentials
+	 */
+	bool EstablishConnection(const FOneSkyConnectionInfo& InConnectionInfo);
+
+public:
+
+	/** true if connection was successfully established */
+	bool bEstablishedConnection;
+};
+
+/**
+ * Connection that is used within specific scope
+ */
+class FScopedOneSkyConnection
+{
+public:
+	/** 
+	 * Constructor - establish a connection.
+	 * The concurrency of the passed in command determines whether the persistent connection is 
+	 * used or another connection is established (connections cannot safely be used across different
+	 * threads).
+	 */
+	FScopedOneSkyConnection( const class FOneSkyLocalizationServiceCommand& InCommand );
+
+	/** 
+	 * Constructor - establish a connection.
+	 * The concurrency passed in determines whether the persistent connection is used or another 
+	 * connection is established (connections cannot safely be used across different threads).
+	 */
+	FScopedOneSkyConnection( ELocalizationServiceOperationConcurrency::Type InConcurrency, const FOneSkyConnectionInfo& InConnectionInfo );
+
+	/**
+	 * Destructor - disconnect if this is a temporary connection
+	 */
+	~FScopedOneSkyConnection();
+
+	/** Get the connection wrapped by this class */
+	FOneSkyConnection& GetConnection() const
+	{
+		return *Connection;
+	}
+
+	/** Check if this connection is valid */
+	bool IsValid() const
+	{
+		return Connection != NULL;
+	}
+
+private:
+	/** Set up the connection */
+	void Initialize( const FOneSkyConnectionInfo& InConnectionInfo );
+
+private:
+	/** The OneSky connection we are using */
+	FOneSkyConnection* Connection;
+
+	/** The concurrency of this connection */
+	ELocalizationServiceOperationConcurrency::Type Concurrency;
+};
\ No newline at end of file
diff --git a/Engine/Plugins/Developer/OneSkyLocalizationService/Source/OneSkyLocalizationService/Private/OneSkyConnectionInfo.h b/Engine/Plugins/Developer/OneSkyLocalizationService/Source/OneSkyLocalizationService/Private/OneSkyConnectionInfo.h
new file mode 100644
--- /dev/null
+++ b/Engine/Plugins/Developer/OneSkyLocalizationService/Source/OneSkyLocalizationService/Private/OneSkyConnectionInfo.h
@@ -0,0 +1,16 @@
+// Copyright 1998-2015 Epic Games, Inc. All Rights Reserved.
+
+#pragma once
+
+/** Helper struct defining parameters we need to establish a connection */
+struct FOneSkyConnectionInfo
+{
+	/** Name of this connection configuration */
+	FString Name;
+
+	/** OneSky ApiKey */
+	FString ApiKey;
+
+	/** OneSky ApiSecret */
+	FString ApiSecret;
+};
\ No newline at end of file
diff --git a/Engine/Plugins/Developer/OneSkyLocalizationService/Source/OneSkyLocalizationService/Private/OneSkyLocalizationServiceCommand.cpp b/Engine/Plugins/Developer/OneSkyLocalizationService/Source/OneSkyLocalizationService/Private/OneSkyLocalizationServiceCommand.cpp
new file mode 100644
--- /dev/null
+++ b/Engine/Plugins/Developer/OneSkyLocalizationService/Source/OneSkyLocalizationService/Private/OneSkyLocalizationServiceCommand.cpp
@@ -0,0 +1,59 @@
+// Copyright 1998-2015 Epic Games, Inc. All Rights Reserved.
+
+#include "OneSkyLocalizationServicePrivatePCH.h"
+#include "OneSkyLocalizationServiceCommand.h"
+#include "OneSkyLocalizationServiceModule.h"
+#include "IOneSkyLocalizationServiceWorker.h"
+
+FOneSkyLocalizationServiceCommand::FOneSkyLocalizationServiceCommand(const TSharedRef<class ILocalizationServiceOperation, ESPMode::ThreadSafe>& InOperation, const TSharedRef<class IOneSkyLocalizationServiceWorker, ESPMode::ThreadSafe>& InWorker, const FLocalizationServiceOperationComplete& InOperationCompleteDelegate /*= FTranslationServiceOperationComplete() */)
+	: Operation(InOperation)
+	, Worker(InWorker)
+	, OperationCompleteDelegate(InOperationCompleteDelegate)
+	, bExecuteProcessed(0)
+	, bCancelled(0)
+	, bCommandSuccessful(false)
+	, bConnectionDropped(false)
+	, bAutoDelete(true)
+	, Concurrency(ELocalizationServiceOperationConcurrency::Synchronous)
+{
+	// grab the providers settings here, so we don't access them once the worker thread is launched
+	check(IsInGameThread());
+	FOneSkyLocalizationServiceModule& OneSkyLocalizationService = FModuleManager::LoadModuleChecked<FOneSkyLocalizationServiceModule>( "OneSkyLocalizationService" );
+	ConnectionInfo = OneSkyLocalizationService.AccessSettings().GetConnectionInfo();
+}
+
+bool FOneSkyLocalizationServiceCommand::DoWork()
+{
+	bCommandSuccessful = Worker->Execute(*this);
+
+	// This is not set here so we can wait for the worker's HTTP response delegate to be called 
+	// FPlatformAtomics::InterlockedExchange(&bExecuteProcessed, 1);
+
+	while (!bExecuteProcessed)
+	{
+		FPlatformProcess::Sleep(0.01f);
+	}
+
+	return bCommandSuccessful;
+}
+
+void FOneSkyLocalizationServiceCommand::Abandon()
+{
+	FPlatformAtomics::InterlockedExchange(&bExecuteProcessed, 1);
+}
+
+void FOneSkyLocalizationServiceCommand::DoThreadedWork()
+{
+	Concurrency = ELocalizationServiceOperationConcurrency::Asynchronous;
+	DoWork();
+}
+
+void FOneSkyLocalizationServiceCommand::Cancel()
+{
+	FPlatformAtomics::InterlockedExchange(&bCancelled, 1);
+}
+
+bool FOneSkyLocalizationServiceCommand::IsCanceled() const
+{
+	return bCancelled != 0;
+}
\ No newline at end of file
diff --git a/Engine/Plugins/Developer/OneSkyLocalizationService/Source/OneSkyLocalizationService/Private/OneSkyLocalizationServiceCommand.h b/Engine/Plugins/Developer/OneSkyLocalizationService/Source/OneSkyLocalizationService/Private/OneSkyLocalizationServiceCommand.h
new file mode 100644
--- /dev/null
+++ b/Engine/Plugins/Developer/OneSkyLocalizationService/Source/OneSkyLocalizationService/Private/OneSkyLocalizationServiceCommand.h
@@ -0,0 +1,89 @@
+// Copyright 1998-2015 Epic Games, Inc. All Rights Reserved.
+
+#pragma once
+
+#include "OneSkyConnectionInfo.h"
+
+DECLARE_DELEGATE_RetVal(bool, FOnIsCancelled);
+
+/**
+ * Used to execute OneSky commands multi-threaded.
+ */
+class FOneSkyLocalizationServiceCommand : public IQueuedWork
+{
+public:
+
+	FOneSkyLocalizationServiceCommand(const TSharedRef<class ILocalizationServiceOperation, ESPMode::ThreadSafe>& InOperation, const TSharedRef<class IOneSkyLocalizationServiceWorker, ESPMode::ThreadSafe>& InWorker, const FLocalizationServiceOperationComplete& InOperationCompleteDelegate = FLocalizationServiceOperationComplete() );
+
+	/**
+	 * This is where the real thread work is done. All work that is done for
+	 * this queued object should be done from within the call to this function.
+	 */
+	bool DoWork();
+
+	/**
+	 * Tells the queued work that it is being abandoned so that it can do
+	 * per object clean up as needed. This will only be called if it is being
+	 * abandoned before completion. NOTE: This requires the object to delete
+	 * itself using whatever heap it was allocated in.
+	 */
+	virtual void Abandon() override;
+
+	/**
+	 * This method is also used to tell the object to cleanup but not before
+	 * the object has finished it's work.
+	 */ 
+	virtual void DoThreadedWork() override;
+
+	/** Attempt to cancel the operation */
+	void Cancel();
+
+	/** Is the operation canceled? */
+	bool IsCanceled() const;
+
+public:
+	/** Connection parameters, reproduced here because if is not safe to access the provider's settings from another thread */
+	FOneSkyConnectionInfo ConnectionInfo;
+
+	/** Operation we want to perform - contains outward-facing parameters & results */
+	TSharedRef<class ILocalizationServiceOperation, ESPMode::ThreadSafe> Operation;
+
+	/** The object that will actually do the work */
+	TSharedRef<class IOneSkyLocalizationServiceWorker, ESPMode::ThreadSafe> Worker;
+
+	/** Delegate to notify when this operation completes */
+	FLocalizationServiceOperationComplete OperationCompleteDelegate;
+
+	/**If true, this command has been processed by the Localization service thread*/
+	volatile int32 bExecuteProcessed;
+
+	/**If true, this command has been cancelled*/
+	volatile int32 bCancelled;
+
+	/**If true, the Localization service command succeeded*/
+	bool bCommandSuccessful;
+
+	/**If true, the Localization service connection was dropped while this command was being executed*/
+	bool bConnectionDropped;
+
+	/** If true, this command will be automatically cleaned up in Tick() */
+	bool bAutoDelete;
+
+	/** Whether we are running multi-treaded or not*/
+	ELocalizationServiceOperationConcurrency::Type Concurrency;
+
+	/** The GUID of the Localization Target we are working with */
+	FGuid TargetGuid;
+
+	/** Files to perform this operation on */
+	TArray< FString > Files;
+
+	/** Translations to perform this operation on */
+	TArray< FLocalizationServiceTranslationIdentifier > Translations;
+
+	/**Info and/or warning message message storage*/
+	TArray< FText > InfoMessages;
+
+	/**Potential error message storage*/
+	TArray< FText > ErrorMessages;
+};
\ No newline at end of file
diff --git a/Engine/Plugins/Developer/OneSkyLocalizationService/Source/OneSkyLocalizationService/Private/OneSkyLocalizationServiceModule.cpp b/Engine/Plugins/Developer/OneSkyLocalizationService/Source/OneSkyLocalizationService/Private/OneSkyLocalizationServiceModule.cpp
new file mode 100644
--- /dev/null
+++ b/Engine/Plugins/Developer/OneSkyLocalizationService/Source/OneSkyLocalizationService/Private/OneSkyLocalizationServiceModule.cpp
@@ -0,0 +1,85 @@
+// Copyright 1998-2015 Epic Games, Inc. All Rights Reserved.
+
+#include "OneSkyLocalizationServicePrivatePCH.h"
+#include "OneSkyLocalizationServiceModule.h"
+#include "ModuleManager.h"
+#include "Developer/LocalizationService/Public/ILocalizationServiceModule.h"
+#include "OneSkyLocalizationServiceOperations.h"
+#include "OneSkyLocalizationServiceState.h"
+#include "OneSkyLocalizationServiceSettings.h"
+#include "Runtime/Core/Public/Features/IModularFeatures.h"
+
+#define LOCTEXT_NAMESPACE "OneSkyTranslationService"
+
+template<typename Type>
+static TSharedRef<IOneSkyLocalizationServiceWorker, ESPMode::ThreadSafe> CreateWorker()
+{
+	return MakeShareable( new Type() );
+}
+
+void FOneSkyLocalizationServiceModule::StartupModule()
+{
+	// Register our operations
+
+	// ProjectGroup API
+	OneSkyLocalizationServiceProvider.RegisterWorker("ListProjectGroups", FGetOneSkyLocalizationServiceWorker::CreateStatic(&CreateWorker<FOneSkyListProjectGroupsWorker>));
+	OneSkyLocalizationServiceProvider.RegisterWorker("ShowProjectGroup", FGetOneSkyLocalizationServiceWorker::CreateStatic(&CreateWorker<FOneSkyShowProjectGroupWorker>));
+	OneSkyLocalizationServiceProvider.RegisterWorker("CreateProjectGroup", FGetOneSkyLocalizationServiceWorker::CreateStatic(&CreateWorker<FOneSkyCreateProjectGroupWorker>));
+	OneSkyLocalizationServiceProvider.RegisterWorker("ListProjectGroupLanguages", FGetOneSkyLocalizationServiceWorker::CreateStatic(&CreateWorker<FOneSkyListProjectGroupLanguagesWorker>));
+
+	// Project API
+	OneSkyLocalizationServiceProvider.RegisterWorker("ListProjectsInGroup", FGetOneSkyLocalizationServiceWorker::CreateStatic(&CreateWorker<FOneSkyListProjectsInGroupWorker>));
+	OneSkyLocalizationServiceProvider.RegisterWorker("ShowProject", FGetOneSkyLocalizationServiceWorker::CreateStatic(&CreateWorker<FOneSkyShowProjectWorker>));
+	OneSkyLocalizationServiceProvider.RegisterWorker("CreateProject", FGetOneSkyLocalizationServiceWorker::CreateStatic(&CreateWorker<FOneSkyCreateProjectWorker>));
+	OneSkyLocalizationServiceProvider.RegisterWorker("ListProjectLanguages", FGetOneSkyLocalizationServiceWorker::CreateStatic(&CreateWorker<FOneSkyListProjectLanguagesWorker>));
+
+	// Translation API
+	OneSkyLocalizationServiceProvider.RegisterWorker("TranslationStatus", FGetOneSkyLocalizationServiceWorker::CreateStatic(&CreateWorker<FOneSkyTranslationStatusWorker>));
+	OneSkyLocalizationServiceProvider.RegisterWorker("TranslationExport", FGetOneSkyLocalizationServiceWorker::CreateStatic(&CreateWorker<FOneSkyTranslationExportWorker>));
+
+	// Files API
+	OneSkyLocalizationServiceProvider.RegisterWorker("ListUploadedFiles", FGetOneSkyLocalizationServiceWorker::CreateStatic(&CreateWorker<FOneSkyListUploadedFilesWorker>));
+	OneSkyLocalizationServiceProvider.RegisterWorker("UploadFile", FGetOneSkyLocalizationServiceWorker::CreateStatic(&CreateWorker<FOneSkyUploadFileWorker>));
+
+	// Phrase Collections API
+	OneSkyLocalizationServiceProvider.RegisterWorker("ListPhraseCollections", FGetOneSkyLocalizationServiceWorker::CreateStatic(&CreateWorker<FOneSkyListPhraseCollectionsWorker>));
+	//OneSkyLocalizationServiceProvider.RegisterWorker("ShowPhraseCollection", FGetOneSkyLocalizationServiceWorker::CreateStatic(&CreateWorker<FOneSkyShowPhraseCollectionWorker>));
+	//OneSkyLocalizationServiceProvider.RegisterWorker("ImportPhraseCollections", FGetOneSkyLocalizationServiceWorker::CreateStatic(&CreateWorker<FOneSkyImportPhraseCollectionsWorker>));
+
+	// ProjectTypes API
+	//OneSkyLocalizationServiceProvider.RegisterWorker("ListProjectTypes", FGetOneSkyLocalizationServiceWorker::CreateStatic(&CreateWorker<FOneSkyListProjectTypesWorker>));
+
+	// load our settings
+	OneSkyLocalizationServiceSettings.LoadSettings();
+
+	// Bind our Localization service provider to the editor
+	IModularFeatures::Get().RegisterModularFeature( "LocalizationService", &OneSkyLocalizationServiceProvider );
+}
+
+void FOneSkyLocalizationServiceModule::ShutdownModule()
+{
+	// shut down the provider, as this module is going away
+	OneSkyLocalizationServiceProvider.Close();
+
+	// unbind provider from editor
+	IModularFeatures::Get().UnregisterModularFeature( "LocalizationService", &OneSkyLocalizationServiceProvider );
+}
+
+FOneSkyLocalizationServiceSettings& FOneSkyLocalizationServiceModule::AccessSettings()
+{
+	return OneSkyLocalizationServiceSettings;
+}
+
+void FOneSkyLocalizationServiceModule::SaveSettings()
+{
+	if (FApp::IsUnattended() || IsRunningCommandlet())
+	{
+		return;
+	}
+
+	OneSkyLocalizationServiceSettings.SaveSettings();
+}
+
+IMPLEMENT_MODULE(FOneSkyLocalizationServiceModule, OneSkyLocalizationService);
+
+#undef LOCTEXT_NAMESPACE
diff --git a/Engine/Plugins/Developer/OneSkyLocalizationService/Source/OneSkyLocalizationService/Private/OneSkyLocalizationServiceModule.h b/Engine/Plugins/Developer/OneSkyLocalizationService/Source/OneSkyLocalizationService/Private/OneSkyLocalizationServiceModule.h
new file mode 100644
--- /dev/null
+++ b/Engine/Plugins/Developer/OneSkyLocalizationService/Source/OneSkyLocalizationService/Private/OneSkyLocalizationServiceModule.h
@@ -0,0 +1,44 @@
+// Copyright 1998-2015 Epic Games, Inc. All Rights Reserved.
+
+#pragma once
+
+#include "OneSkyLocalizationServiceProvider.h"
+#include "OneSkyLocalizationServiceSettings.h"
+
+class FOneSkyLocalizationServiceModule : public IModuleInterface
+{
+public:
+	/** IModuleInterface implementation */
+	virtual void StartupModule() override;
+	virtual void ShutdownModule() override;
+
+	/** Access the OneSky Localization service settings */
+	FOneSkyLocalizationServiceSettings& AccessSettings();
+
+	/** Save the OneSky Localization service settings */
+	void SaveSettings();
+
+	/** Access the one and only Perforce provider */
+	FOneSkyLocalizationServiceProvider& GetProvider()
+	{
+		return OneSkyLocalizationServiceProvider;
+	}
+	
+	/**
+	 * Singleton-like access to this module's interface.  This is just for convenience!
+	 * Beware of calling this during the shutdown phase, though.  Your module might have been unloaded already.
+	 *
+	 * @return Returns singleton instance, loading the module on demand if needed
+	 */
+	static inline FOneSkyLocalizationServiceModule& Get()
+	{
+		return FModuleManager::LoadModuleChecked< FOneSkyLocalizationServiceModule >("OneSkyLocalizationService");
+	}
+
+private:
+	/** The one and only OneSky Localization service provider */
+	FOneSkyLocalizationServiceProvider OneSkyLocalizationServiceProvider;
+
+	/** The settings for OneSky Localization service */
+	FOneSkyLocalizationServiceSettings OneSkyLocalizationServiceSettings;
+};
\ No newline at end of file
diff --git a/Engine/Plugins/Developer/OneSkyLocalizationService/Source/OneSkyLocalizationService/Private/OneSkyLocalizationServiceOperations.cpp b/Engine/Plugins/Developer/OneSkyLocalizationService/Source/OneSkyLocalizationService/Private/OneSkyLocalizationServiceOperations.cpp
new file mode 100644
--- /dev/null
+++ b/Engine/Plugins/Developer/OneSkyLocalizationService/Source/OneSkyLocalizationService/Private/OneSkyLocalizationServiceOperations.cpp
@@ -0,0 +1,1092 @@
+// Copyright 1998-2015 Epic Games, Inc. All Rights Reserved.
+
+#include "OneSkyLocalizationServicePrivatePCH.h"
+#include "OneSkyLocalizationServiceOperations.h"
+#include "OneSkyLocalizationServiceState.h"
+#include "OneSkyLocalizationServiceCommand.h"
+#include "OneSkyConnection.h"
+#include "OneSkyLocalizationServiceModule.h"
+#include "OneSkyLocalizationServiceRevision.h"
+#include "SecureHash.h"
+#include "Runtime/Online/HTTP/Public/Http.h"
+#include "JsonStructDeserializerBackend.h"
+#include "StructDeserializer.h"
+#include "OneSkyLocalizationServiceResponseTypes.h"
+
+#define LOCTEXT_NAMESPACE "OneSkyLocalizationService"
+
+FString GetAuthenticationParameters(const FOneSkyConnectionInfo& ConnectionInfo)
+{
+	//[REDACTED_PROFILE_URL]/api-documentation-platform/blob/master/README.md#authentication
+
+	FString UtcNowUnixTimeStamp = FString::FromInt(FDateTime::UtcNow().ToUnixTimestamp());
+	FString Hash = FMD5::HashAnsiString(*(UtcNowUnixTimeStamp + ConnectionInfo.ApiSecret));
+
+	return "api_key=" + ConnectionInfo.ApiKey + "&dev_hash=" + Hash + "&timestamp=" + UtcNowUnixTimeStamp;
+}
+
+FString AddAuthenticationParameters(const FOneSkyConnectionInfo& ConnectionInfo, const FString& Url)
+{
+	return Url + "?" + GetAuthenticationParameters(ConnectionInfo);
+}
+
+static bool UpdateCachedLocalizationStates(const TMap<FLocalizationServiceTranslationIdentifier, TSharedRef<FOneSkyLocalizationServiceState, ESPMode::ThreadSafe>, FDefaultSetAllocator, FLocalizationServiceTranslationIdentifierKeyFuncs<TSharedRef<FOneSkyLocalizationServiceState, ESPMode::ThreadSafe>>>& InResults)
+{
+	FOneSkyLocalizationServiceModule& OneSkyLocalizationService = FOneSkyLocalizationServiceModule::Get();
+	for (auto It = InResults.CreateConstIterator(); It; ++It)
+	{
+		TSharedRef<FOneSkyLocalizationServiceState, ESPMode::ThreadSafe> State = OneSkyLocalizationService.GetProvider().GetStateInternal(It.Key());
+		State->SetState(It.Value()->GetState());
+		State->SetTranslation(It.Value()->GetTranslationString());
+		State->TimeStamp = FDateTime::Now();
+	}
+
+	return InResults.Num() > 0;
+}
+
+static FString GetFileFormat(FString FileExtension)
+{
+	if (FileExtension == ".po")
+	{
+		return "GNU_PO";
+	}
+	else if (FileExtension == ".pot")
+	{
+		return "GNU_POT";
+	}
+	else
+	{
+		return "";
+	}
+}
+
+static bool DeserializeResponseToStruct(void* OutStruct, UStruct& TypeInfo, FHttpResponsePtr HttpResponse)
+{
+	bool bResult = false;
+	FString ResponseStr = HttpResponse->GetContentAsString();
+	FText ErrorText;
+
+	if (HttpResponse.IsValid())
+	{
+		if (EHttpResponseCodes::IsOk(HttpResponse->GetResponseCode()))
+		{
+
+			// Write string to FMemoryWriter to force unicode
+			TArray<uint8> InBytes;
+			FMemoryWriter Writer(InBytes);
+			Writer.ArForceUnicode = true;
+			Writer << ResponseStr;
+
+			FMemoryReader Reader(InBytes);
+			// FMemoryWriter writes size of string at beginning, need to ignore this or json parsing errors occur
+			Reader.Seek(4); 
+
+			FJsonStructDeserializerBackend Backend(Reader);
+
+			bResult = FStructDeserializer::Deserialize(OutStruct, TypeInfo, Backend);
+		}
+		else
+		{
+			ErrorText = FText::Format(LOCTEXT("InvalidResponse", "Invalid response. code={0} error={1}"), FText::FromString(FString::FromInt(HttpResponse->GetResponseCode())), FText::FromString(ResponseStr));
+		}
+	}
+
+	if (!bResult)
+	{
+		UE_LOG(LogLocalizationService, Warning, TEXT("%s"), *(ErrorText.ToString()));
+	}
+
+	return bResult;
+}
+
+FName FOneSkyConnect::GetName() const
+{
+	return "FOneSkyConnect";
+}
+
+bool FOneSkyConnect::Execute(class FOneSkyLocalizationServiceCommand& InCommand)
+{
+	//FOneSkyLocalizationServiceModule& OneSkyLocalizationService = FModuleManager::LoadModuleChecked<FOneSkyLocalizationServiceModule>("OneSkyLocalizationService");
+	//OneSkyLocalizationService.GetProvider().RetrieveProjectGroups();
+
+	return true;
+}
+
+
+// LIST PROJECT GROUPS
+
+bool FOneSkyListProjectGroupsWorker::Execute(FOneSkyLocalizationServiceCommand& InCommand)
+{
+	// Store pointer to command so we can access it in the http request callback
+	Command = &InCommand;
+	TSharedPtr<FOneSkyListProjectGroupsOperation, ESPMode::ThreadSafe> ListProjectGroupsOp = StaticCastSharedRef<FOneSkyListProjectGroupsOperation>(InCommand.Operation);
+	
+	int32 InStartPage = -1;
+	int32 InItemsPerPage = -1;
+	
+	if (ListProjectGroupsOp.IsValid())
+	{
+		InStartPage = ListProjectGroupsOp->GetInStartPage();
+		InItemsPerPage = ListProjectGroupsOp->GetInItemsPerPage();
+	}
+
+	FString Url = AddAuthenticationParameters(InCommand.ConnectionInfo, "https://platform.api.onesky.io/1/project-groups");
+	if (InStartPage != -1 && InItemsPerPage != -1)
+	{
+		Url += FString(TEXT("&page=")) +FString::FromInt(InStartPage) + FString(TEXT("&per_page=")) + FString::FromInt(InItemsPerPage);
+	}
+
+	FHttpModule& HttpModule = FModuleManager::LoadModuleChecked<FHttpModule>("HTTP");
+	TSharedRef<class IHttpRequest> HttpRequest = HttpModule.Get().CreateRequest();
+
+	// kick off http request to read
+	HttpRequest->OnProcessRequestComplete().BindRaw(this, &FOneSkyListProjectGroupsWorker::Query_HttpRequestComplete);
+	HttpRequest->SetURL(Url);
+	HttpRequest->SetHeader(TEXT("Content-Type"), TEXT("application/json; charset=utf-8"));
+	HttpRequest->SetVerb(TEXT("GET"));
+	HttpRequest->ProcessRequest();
+
+	return InCommand.bCommandSuccessful;
+}
+
+
+
+void FOneSkyListProjectGroupsWorker::Query_HttpRequestComplete(FHttpRequestPtr HttpRequest, FHttpResponsePtr HttpResponse, bool bSucceeded)
+{
+	bool bResult = false;
+
+	if (bSucceeded)
+	{
+			UScriptStruct& StructType = *(FOneSkyListProjectGroupsResponse::StaticStruct());
+			bResult = DeserializeResponseToStruct(&OutListProjectGroupsResponse, StructType, HttpResponse);
+	}
+
+	if (Command != nullptr)
+	{
+		Command->bCommandSuccessful = bResult;
+		FPlatformAtomics::InterlockedExchange(&(Command->bExecuteProcessed), 1);
+	}
+}
+
+FName FOneSkyListProjectGroupsWorker::GetName() const
+{
+	return "ListProjectGroupsWorker";
+}
+
+// end LIST PROJECT GROUPS
+
+
+// SHOW PROJECT GROUP
+
+FName FOneSkyShowProjectGroupWorker::GetName() const
+{
+	return "ShowProjectGroupWorker";
+}
+
+bool FOneSkyShowProjectGroupWorker::Execute(class FOneSkyLocalizationServiceCommand& InCommand)
+{
+	// Store pointer to command so we can access it in the http request callback
+	Command = &InCommand;
+
+	TSharedPtr<FOneSkyShowProjectGroupOperation, ESPMode::ThreadSafe> ShowProjectGroupOp = StaticCastSharedRef<FOneSkyShowProjectGroupOperation>(InCommand.Operation);
+
+	int32 InProjectGroupId = -1;
+
+	if (ShowProjectGroupOp.IsValid())
+	{
+		InProjectGroupId = ShowProjectGroupOp->GetInProjectGroupId();
+	}
+
+	FString Url = "https://platform.api.onesky.io/1/project-groups/";
+	Url += FString::FromInt(InProjectGroupId);
+	Url = AddAuthenticationParameters(InCommand.ConnectionInfo, Url);
+
+	FHttpModule& HttpModule = FModuleManager::LoadModuleChecked<FHttpModule>("HTTP");
+	TSharedRef<class IHttpRequest> HttpRequest = HttpModule.Get().CreateRequest();
+
+	// kick off http request to read
+	HttpRequest->OnProcessRequestComplete().BindRaw(this, &FOneSkyShowProjectGroupWorker::Query_HttpRequestComplete);
+	HttpRequest->SetURL(Url);
+	HttpRequest->SetHeader(TEXT("Content-Type"), TEXT("application/json; charset=utf-8"));
+	HttpRequest->SetVerb(TEXT("GET"));
+	HttpRequest->ProcessRequest();
+
+	return InCommand.bCommandSuccessful;
+
+	return false;
+}
+
+void FOneSkyShowProjectGroupWorker::Query_HttpRequestComplete(FHttpRequestPtr HttpRequest, FHttpResponsePtr HttpResponse, bool bSucceeded)
+{
+	bool bResult = false;
+
+	if (bSucceeded)
+	{
+		UScriptStruct& StructType = *(FOneSkyShowProjectGroupResponse::StaticStruct());
+		bResult = DeserializeResponseToStruct(&OutShowProjectGroupResponse, StructType, HttpResponse);
+	}
+
+	if (Command != nullptr)
+	{
+		Command->bCommandSuccessful = bResult;
+		FPlatformAtomics::InterlockedExchange(&(Command->bExecuteProcessed), 1);
+	}
+}
+
+// end SHOW PROJECT GROUP
+
+
+// CREATE PROJECT GROUP
+
+FName FOneSkyCreateProjectGroupWorker::GetName() const
+{
+	return "CreateProjectGroupWorker";
+}
+
+bool FOneSkyCreateProjectGroupWorker::Execute(class FOneSkyLocalizationServiceCommand& InCommand)
+{
+	// Store pointer to command so we can access it in the http request callback
+	Command = &InCommand;
+
+	TSharedPtr<FOneSkyCreateProjectGroupOperation, ESPMode::ThreadSafe> CreateProjectGroupOp = StaticCastSharedRef<FOneSkyCreateProjectGroupOperation>(InCommand.Operation);
+
+	FString InProjectGroupName = "";
+	FString InProjectGroupBaseLocale = "";
+
+	if (CreateProjectGroupOp.IsValid())
+	{
+		//we need to url encode for special characters (especially other languages)
+		InProjectGroupName = FGenericPlatformHttp::UrlEncode(CreateProjectGroupOp->GetInProjectGroupName());
+		InProjectGroupBaseLocale = FGenericPlatformHttp::UrlEncode(CreateProjectGroupOp->GetInProjectGroupBaseLocale());
+	}
+
+	FString Url = "https://platform.api.onesky.io/1/project-groups";
+	FString Parameters = GetAuthenticationParameters(InCommand.ConnectionInfo) + FString(TEXT("&name=")) + InProjectGroupName + FString(TEXT("&locale=")) + InProjectGroupBaseLocale;
+	Url += "?" + Parameters;
+
+	FHttpModule& HttpModule = FModuleManager::LoadModuleChecked<FHttpModule>("HTTP");
+	TSharedRef<class IHttpRequest> HttpRequest = HttpModule.Get().CreateRequest();
+
+	// Seems to work fine with no content payload (same info is in the URL anyway)
+	// Leaving code here as a reference for later
+	// Write string to FMemoryWriter to force unicode
+	//TArray<uint8> InBytes;
+	//FMemoryWriter Writer(InBytes);
+	//Writer.ArForceUnicode = true;
+	//Writer << Parameters;
+	//FMemoryReader Reader(InBytes);
+	//// FMemoryWriter writes size of string at beginning, need to ignore this or parsing errors occur
+	//InBytes.RemoveAt(3);
+	//InBytes.RemoveAt(2);
+	//InBytes.RemoveAt(1);
+	//InBytes.RemoveAt(0);
+	//HttpRequest->SetContent(InBytes);
+
+	
+	HttpRequest->OnProcessRequestComplete().BindRaw(this, &FOneSkyCreateProjectGroupWorker::Query_HttpRequestComplete);
+	HttpRequest->SetURL(Url);
+
+	// UTF-8
+	// Not sure if these matter if we don't send a content payload
+	HttpRequest->SetHeader(TEXT("Content-Type"), TEXT("text/html; charset=utf-8"));
+
+	HttpRequest->SetVerb(TEXT("POST"));
+	// kick off http request to read 
+	HttpRequest->ProcessRequest();
+
+	return InCommand.bCommandSuccessful;
+
+	return false;
+}
+
+void FOneSkyCreateProjectGroupWorker::Query_HttpRequestComplete(FHttpRequestPtr HttpRequest, FHttpResponsePtr HttpResponse, bool bSucceeded)
+{
+	bool bResult = false;
+
+	if (bSucceeded)
+	{
+		UScriptStruct& StructType = *(FOneSkyCreateProjectGroupResponse::StaticStruct());
+		bResult = DeserializeResponseToStruct(&OutCreateProjectGroupResponse, StructType, HttpResponse);
+	}
+
+	if (Command != nullptr)
+	{
+		Command->bCommandSuccessful = bResult;
+		FPlatformAtomics::InterlockedExchange(&(Command->bExecuteProcessed), 1);
+	}
+}
+
+// end CREATE PROJECT GROUP
+
+
+// LIST PROJECT GROUP LANGUAGES
+
+FName FOneSkyListProjectGroupLanguagesWorker::GetName() const
+{
+	return "ListProjectGroupLanguagesWorker";
+}
+
+bool FOneSkyListProjectGroupLanguagesWorker::Execute(class FOneSkyLocalizationServiceCommand& InCommand)
+{
+	// Store pointer to command so we can access it in the http request callback
+	Command = &InCommand;
+
+	TSharedPtr<FOneSkyListProjectGroupLanguagesOperation, ESPMode::ThreadSafe> ListProjectGroupLanguagesOp = StaticCastSharedRef<FOneSkyListProjectGroupLanguagesOperation>(InCommand.Operation);
+
+	int32 InProjectGroupId = -1;
+
+	if (ListProjectGroupLanguagesOp.IsValid())
+	{
+		InProjectGroupId = ListProjectGroupLanguagesOp->GetInProjectGroupId();
+	}
+
+	FString Url = "https://platform.api.onesky.io/1/project-groups/" + FString::FromInt(InProjectGroupId) / "languages";
+	Url = AddAuthenticationParameters(InCommand.ConnectionInfo, Url);
+
+	FHttpModule& HttpModule = FModuleManager::LoadModuleChecked<FHttpModule>("HTTP");
+	TSharedRef<class IHttpRequest> HttpRequest = HttpModule.Get().CreateRequest();
+
+	// kick off http request to read
+	HttpRequest->OnProcessRequestComplete().BindRaw(this, &FOneSkyListProjectGroupLanguagesWorker::Query_HttpRequestComplete);
+	HttpRequest->SetURL(Url);
+	HttpRequest->SetHeader(TEXT("Content-Type"), TEXT("application/json; charset=utf-8"));
+	HttpRequest->SetVerb(TEXT("GET"));
+	HttpRequest->ProcessRequest();
+
+	return InCommand.bCommandSuccessful;
+
+	return false;
+}
+
+void FOneSkyListProjectGroupLanguagesWorker::Query_HttpRequestComplete(FHttpRequestPtr HttpRequest, FHttpResponsePtr HttpResponse, bool bSucceeded)
+{
+	bool bResult = false;
+
+	if (bSucceeded)
+	{
+		UScriptStruct& StructType = *(FOneSkyListProjectGroupLanguagesResponse::StaticStruct());
+		bResult = DeserializeResponseToStruct(&OutListProjectGroupLanguagesResponse, StructType, HttpResponse);
+	}
+
+	if (Command != nullptr)
+	{
+		Command->bCommandSuccessful = bResult;
+		FPlatformAtomics::InterlockedExchange(&(Command->bExecuteProcessed), 1);
+	}
+}
+
+// end LIST PROJECT GROUP LANGUAGES
+
+
+
+// LIST PROJECTS IN GROUP
+
+FName FOneSkyListProjectsInGroupWorker::GetName() const
+{
+	return FName(TEXT("FOneSkyListProjectsInGroupWorker"));
+}
+
+bool FOneSkyListProjectsInGroupWorker::Execute(class FOneSkyLocalizationServiceCommand& InCommand)
+{
+	// Store pointer to command so we can access it in the http request callback
+	Command = &InCommand;
+	TSharedPtr<FOneSkyListProjectsInGroupOperation, ESPMode::ThreadSafe> ListProjectsInGroupOp = StaticCastSharedRef<FOneSkyListProjectsInGroupOperation>(InCommand.Operation);
+
+	int32 InProjectGroupId = -1;
+
+	if (ListProjectsInGroupOp.IsValid())
+	{
+		InProjectGroupId = ListProjectsInGroupOp->GetInProjectGroupId();
+	}
+
+	FString Url = "https://platform.api.onesky.io/1/project-groups/" + FString::FromInt(InProjectGroupId) / "projects";
+	Url = AddAuthenticationParameters(InCommand.ConnectionInfo, Url);
+
+	FHttpModule& HttpModule = FModuleManager::LoadModuleChecked<FHttpModule>("HTTP");
+	TSharedRef<class IHttpRequest> HttpRequest = HttpModule.Get().CreateRequest();
+
+	// kick off http request to read
+	HttpRequest->OnProcessRequestComplete().BindRaw(this, &FOneSkyListProjectsInGroupWorker::Query_HttpRequestComplete);
+	HttpRequest->SetURL(Url);
+	HttpRequest->SetHeader(TEXT("Content-Type"), TEXT("application/json; charset=utf-8"));
+	HttpRequest->SetVerb(TEXT("GET"));
+	HttpRequest->ProcessRequest();
+
+	return InCommand.bCommandSuccessful;
+}
+
+void FOneSkyListProjectsInGroupWorker::Query_HttpRequestComplete(FHttpRequestPtr HttpRequest, FHttpResponsePtr HttpResponse, bool bSucceeded)
+{
+	bool bResult = false;
+
+	if (bSucceeded)
+	{
+		UScriptStruct& StructType = *(FOneSkyListProjectsInGroupResponse::StaticStruct());
+		bResult = DeserializeResponseToStruct(&OutListProjectsInGroupResponse, StructType, HttpResponse);
+	}
+
+	if (Command != nullptr)
+	{
+		Command->bCommandSuccessful = bResult;
+		FPlatformAtomics::InterlockedExchange(&(Command->bExecuteProcessed), 1);
+	}
+}
+
+// end LIST PROJECTS IN GROUP
+
+
+// SHOW PROJECT
+
+FName FOneSkyShowProjectWorker::GetName() const
+{
+	return FName(TEXT("FOneSkyShowProjectWorker"));
+}
+
+bool FOneSkyShowProjectWorker::Execute(class FOneSkyLocalizationServiceCommand& InCommand)
+{
+	// Store pointer to command so we can access it in the http request callback
+	Command = &InCommand;
+
+	TSharedPtr<FOneSkyShowProjectOperation, ESPMode::ThreadSafe> ShowProjectOp = StaticCastSharedRef<FOneSkyShowProjectOperation>(InCommand.Operation);
+
+	int32 InProjectId = -1;
+
+	if (ShowProjectOp.IsValid())
+	{
+		InProjectId = ShowProjectOp->GetInProjectId();
+	}
+
+	FString Url = "https://platform.api.onesky.io/1/projects/" + FString::FromInt(InProjectId);
+	Url = AddAuthenticationParameters(InCommand.ConnectionInfo, Url);
+
+	FHttpModule& HttpModule = FModuleManager::LoadModuleChecked<FHttpModule>("HTTP");
+	TSharedRef<class IHttpRequest> HttpRequest = HttpModule.Get().CreateRequest();
+
+	// kick off http request to read
+	HttpRequest->OnProcessRequestComplete().BindRaw(this, &FOneSkyShowProjectWorker::Query_HttpRequestComplete);
+	HttpRequest->SetURL(Url);
+	HttpRequest->SetHeader(TEXT("Content-Type"), TEXT("application/json; charset=utf-8"));
+	HttpRequest->SetVerb(TEXT("GET"));
+	HttpRequest->ProcessRequest();
+
+	return InCommand.bCommandSuccessful;
+
+	return false;
+}
+
+void FOneSkyShowProjectWorker::Query_HttpRequestComplete(FHttpRequestPtr HttpRequest, FHttpResponsePtr HttpResponse, bool bSucceeded)
+{
+	bool bResult = false;
+
+	if (bSucceeded)
+	{
+		UScriptStruct& StructType = *(FOneSkyShowProjectResponse::StaticStruct());
+		bResult = DeserializeResponseToStruct(&OutShowProjectResponse, StructType, HttpResponse);
+	}
+
+	if (Command != nullptr)
+	{
+		Command->bCommandSuccessful = bResult;
+		FPlatformAtomics::InterlockedExchange(&(Command->bExecuteProcessed), 1);
+	}
+}
+
+// end SHOW PROJECT
+
+
+// CREATE PROJECT
+
+FName FOneSkyCreateProjectWorker::GetName() const
+{
+	return FName(TEXT("FOneSkyCreateProjectWorker"));
+}
+
+bool FOneSkyCreateProjectWorker::Execute(class FOneSkyLocalizationServiceCommand& InCommand)
+{
+	// Store pointer to command so we can access it in the http request callback
+	Command = &InCommand;
+
+	TSharedPtr<FOneSkyCreateProjectOperation, ESPMode::ThreadSafe> CreateProjectOp = StaticCastSharedRef<FOneSkyCreateProjectOperation>(InCommand.Operation);
+
+	FString InProjectGroupName = "";
+	FString InProjectDescription = "";
+	FString InProjectType = "";
+	int32 InProjectGroupId = -1;
+
+	if (CreateProjectOp.IsValid())
+	{
+		//we need to url encode for special characters (especially other languages)
+		InProjectGroupName = FGenericPlatformHttp::UrlEncode(CreateProjectOp->GetInProjectGroupName());
+		InProjectDescription = FGenericPlatformHttp::UrlEncode(CreateProjectOp->GetInProjectDescription());
+		InProjectType = FGenericPlatformHttp::UrlEncode(CreateProjectOp->GetInProjectType());
+		InProjectGroupId = CreateProjectOp->GetInProjectGroupId();
+	}
+
+	FString Url = "https://platform.api.onesky.io/1/project-groups/" + FString::FromInt(InProjectGroupId) + "/projects";
+	FString Parameters = GetAuthenticationParameters(InCommand.ConnectionInfo) + FString(TEXT("&name=")) + InProjectGroupName + FString(TEXT("&description=")) + InProjectDescription + FString(TEXT("&project_type=") + InProjectType);
+	Url += "?" + Parameters;
+
+	FHttpModule& HttpModule = FModuleManager::LoadModuleChecked<FHttpModule>("HTTP");
+	TSharedRef<class IHttpRequest> HttpRequest = HttpModule.Get().CreateRequest();
+
+	// kick off http request to read 
+	HttpRequest->OnProcessRequestComplete().BindRaw(this, &FOneSkyCreateProjectWorker::Query_HttpRequestComplete);
+	HttpRequest->SetURL(Url);
+	HttpRequest->SetHeader(TEXT("Content-Type"), TEXT("text/html; charset=utf-8"));
+	HttpRequest->SetVerb(TEXT("POST"));
+	HttpRequest->ProcessRequest();
+
+	return InCommand.bCommandSuccessful;
+
+	return false;
+}
+
+void FOneSkyCreateProjectWorker::Query_HttpRequestComplete(FHttpRequestPtr HttpRequest, FHttpResponsePtr HttpResponse, bool bSucceeded)
+{
+	bool bResult = false;
+
+	if (bSucceeded)
+	{
+		UScriptStruct& StructType = *(FOneSkyCreateProjectResponse::StaticStruct());
+		bResult = DeserializeResponseToStruct(&OutCreateProjectResponse, StructType, HttpResponse);
+	}
+
+	if (Command != nullptr)
+	{
+		Command->bCommandSuccessful = bResult;
+		FPlatformAtomics::InterlockedExchange(&(Command->bExecuteProcessed), 1);
+	}
+}
+
+// end CREATE PROJECT
+
+
+// LIST PROJECT LANGUAGES
+
+FName FOneSkyListProjectLanguagesWorker::GetName() const
+{
+	return "ListProjectLanguagesWorker";
+}
+
+bool FOneSkyListProjectLanguagesWorker::Execute(class FOneSkyLocalizationServiceCommand& InCommand)
+{
+	// Store pointer to command so we can access it in the http request callback
+	Command = &InCommand;
+
+	TSharedPtr<FOneSkyListProjectLanguagesOperation, ESPMode::ThreadSafe> ListProjectLanguagesOp = StaticCastSharedRef<FOneSkyListProjectLanguagesOperation>(InCommand.Operation);
+
+	int32 InProjectId = -1;
+
+	if (ListProjectLanguagesOp.IsValid())
+	{
+		InProjectId = ListProjectLanguagesOp->GetInProjectId();
+	}
+
+	FString Url = "https://platform.api.onesky.io/1/projects/" + FString::FromInt(InProjectId) / "languages";
+	Url = AddAuthenticationParameters(InCommand.ConnectionInfo, Url);
+
+	FHttpModule& HttpModule = FModuleManager::LoadModuleChecked<FHttpModule>("HTTP");
+	TSharedRef<class IHttpRequest> HttpRequest = HttpModule.Get().CreateRequest();
+
+	// kick off http request to read
+	HttpRequest->OnProcessRequestComplete().BindRaw(this, &FOneSkyListProjectLanguagesWorker::Query_HttpRequestComplete);
+	HttpRequest->SetURL(Url);
+	HttpRequest->SetHeader(TEXT("Content-Type"), TEXT("application/json; charset=utf-8"));
+	HttpRequest->SetVerb(TEXT("GET"));
+	HttpRequest->ProcessRequest();
+
+	return InCommand.bCommandSuccessful;
+
+	return false;
+}
+
+void FOneSkyListProjectLanguagesWorker::Query_HttpRequestComplete(FHttpRequestPtr HttpRequest, FHttpResponsePtr HttpResponse, bool bSucceeded)
+{
+	bool bResult = false;
+
+	if (bSucceeded)
+	{
+		UScriptStruct& StructType = *(FOneSkyListProjectLanguagesResponse::StaticStruct());
+		bResult = DeserializeResponseToStruct(&OutListProjectLanguagesResponse, StructType, HttpResponse);
+	}
+
+	if (Command != nullptr)
+	{
+		Command->bCommandSuccessful = bResult;
+		FPlatformAtomics::InterlockedExchange(&(Command->bExecuteProcessed), 1);
+	}
+}
+
+//end LIST PROJECT LANGUAGES
+
+
+// TRANSLATION STATUS
+
+FName FOneSkyTranslationStatusWorker::GetName() const
+{
+	return "TranslationStatusWorker";
+}
+
+bool FOneSkyTranslationStatusWorker::Execute(class FOneSkyLocalizationServiceCommand& InCommand)
+{
+	// Store pointer to command so we can access it in the http request callback
+	Command = &InCommand;
+	TSharedPtr<FOneSkyTranslationStatusOperation, ESPMode::ThreadSafe> TranslationStatusOp = StaticCastSharedRef<FOneSkyTranslationStatusOperation>(InCommand.Operation);
+
+	int32 InProjectId = -1;
+	FString InFileName = "";
+	FString InLocale = "";
+
+	if (TranslationStatusOp.IsValid())
+	{
+		InProjectId = TranslationStatusOp->GetInProjectId();
+		InFileName = FGenericPlatformHttp::UrlEncode(TranslationStatusOp->GetInFileName());
+		InLocale = FGenericPlatformHttp::UrlEncode(TranslationStatusOp->GetInLocale());
+	}
+	
+	FString Url = AddAuthenticationParameters(InCommand.ConnectionInfo, "https://platform.api.onesky.io/1/projects/" + FString::FromInt(InProjectId) / "translations/status");
+	Url += FString(TEXT("&file_name=")) + InFileName + FString(TEXT("&locale=")) + InLocale;
+
+	FHttpModule& HttpModule = FModuleManager::LoadModuleChecked<FHttpModule>("HTTP");
+	TSharedRef<class IHttpRequest> HttpRequest = HttpModule.Get().CreateRequest();
+
+	// kick off http request to read
+	HttpRequest->OnProcessRequestComplete().BindRaw(this, &FOneSkyTranslationStatusWorker::Query_HttpRequestComplete);
+	HttpRequest->SetURL(Url);
+	HttpRequest->SetHeader(TEXT("Content-Type"), TEXT("application/json; charset=utf-8"));
+	HttpRequest->SetVerb(TEXT("GET"));
+	HttpRequest->ProcessRequest();
+
+	return InCommand.bCommandSuccessful;
+}
+
+void FOneSkyTranslationStatusWorker::Query_HttpRequestComplete(FHttpRequestPtr HttpRequest, FHttpResponsePtr HttpResponse, bool bSucceeded)
+{
+	bool bResult = false;
+
+	if (bSucceeded)
+	{
+		UScriptStruct& StructType = *(FOneSkyTranslationStatusResponse::StaticStruct());
+		bResult = DeserializeResponseToStruct(&OutTranslationStatusResponse, StructType, HttpResponse);
+	}
+
+	if (Command != nullptr)
+	{
+		Command->bCommandSuccessful = bResult;
+		FPlatformAtomics::InterlockedExchange(&(Command->bExecuteProcessed), 1);
+	}
+}
+
+// end TRANSLATION STATUS
+
+
+// TRANSLATION EXPORT
+
+FName FOneSkyTranslationExportWorker::GetName() const
+{
+	return "TranslationExportWorker";
+}
+
+bool FOneSkyTranslationExportWorker::Execute(class FOneSkyLocalizationServiceCommand& InCommand)
+{
+	// Store pointer to command so we can access it in the http request callback
+	Command = &InCommand;
+	TSharedPtr<FOneSkyTranslationExportOperation, ESPMode::ThreadSafe> TranslationExportOp = StaticCastSharedRef<FOneSkyTranslationExportOperation>(InCommand.Operation);
+
+	int32 InProjectId = -1;
+	FString InSourceFileName = "";
+	FString InExportFileName = "";
+	FString InLocale = "";
+
+	if (TranslationExportOp.IsValid())
+	{
+		InProjectId = TranslationExportOp->GetInProjectId();
+		InSourceFileName = FGenericPlatformHttp::UrlEncode(TranslationExportOp->GetInSourceFileName());
+		InExportFileName = FGenericPlatformHttp::UrlEncode(TranslationExportOp->GetInExportFileName());
+		InLocale = FGenericPlatformHttp::UrlEncode(TranslationExportOp->GetInLocale());
+	}
+
+	FString Url = AddAuthenticationParameters(InCommand.ConnectionInfo, "https://platform.api.onesky.io/1/projects/" + FString::FromInt(InProjectId) / "translations");
+	Url += FString(TEXT("&source_file_name=")) + InSourceFileName + FString(TEXT("&locale=")) + InLocale;
+	if (!InExportFileName.IsEmpty())
+	{
+		Url += "&export_file_name=" + InExportFileName;
+	}
+
+	FHttpModule& HttpModule = FModuleManager::LoadModuleChecked<FHttpModule>("HTTP");
+	TSharedRef<class IHttpRequest> HttpRequest = HttpModule.Get().CreateRequest();
+
+	// kick off http request to read
+	HttpRequest->OnProcessRequestComplete().BindRaw(this, &FOneSkyTranslationExportWorker::Query_HttpRequestComplete);
+	HttpRequest->SetURL(Url);
+	HttpRequest->SetHeader(TEXT("Content-Type"), TEXT("application/json; charset=utf-8"));
+	HttpRequest->SetVerb(TEXT("GET"));
+	HttpRequest->ProcessRequest();
+
+	return InCommand.bCommandSuccessful;
+}
+
+void FOneSkyTranslationExportWorker::Query_HttpRequestComplete(FHttpRequestPtr HttpRequest, FHttpResponsePtr HttpResponse, bool bSucceeded)
+{
+	bool bResult = false;
+
+	if (bSucceeded)
+	{
+		FString ResponseStr = HttpResponse->GetContentAsString();
+		FText ErrorText;
+
+		if (HttpResponse.IsValid())
+		{
+			EHttpResponseCodes::Type ResponseCode = (EHttpResponseCodes::Type) HttpResponse->GetResponseCode();
+
+			if (EHttpResponseCodes::IsOk(ResponseCode))
+			{
+				// TODO: Test Accepted case.  Haven't been able to make this happen yet...
+				if (ResponseCode == EHttpResponseCodes::Accepted)
+				{
+					if (Command != nullptr)
+					{
+						TSharedPtr<FOneSkyTranslationExportOperation, ESPMode::ThreadSafe> TranslationExportOp = StaticCastSharedRef<FOneSkyTranslationExportOperation>(Command->Operation);
+						if (TranslationExportOp.IsValid())
+						{
+							// The file is not ready, try again
+							TSharedRef<FOneSkyTranslationExportOperation, ESPMode::ThreadSafe> NewTranslationExportLanguagesOp = ILocalizationServiceOperation::Create<FOneSkyTranslationExportOperation>();
+							NewTranslationExportLanguagesOp->SetInProjectId(TranslationExportOp->GetInProjectId());
+							NewTranslationExportLanguagesOp->SetInSourceFileName(TranslationExportOp->GetInSourceFileName());
+							NewTranslationExportLanguagesOp->SetInLocale(TranslationExportOp->GetInLocale());
+							NewTranslationExportLanguagesOp->SetInExportFileName(TranslationExportOp->GetInExportFileName());
+							ILocalizationServiceModule::Get().GetProvider().Execute(NewTranslationExportLanguagesOp, TArray<FLocalizationServiceTranslationIdentifier>(), ELocalizationServiceOperationConcurrency::Asynchronous);
+						}
+						else
+						{
+							ErrorText = LOCTEXT("TranslationExportQueryFailedTranslationExportOpInvalid", "Translation Export Query Failed: Translation Export Operation is invalid.");
+						}
+
+						return;
+					}
+					else
+					{
+						ErrorText = LOCTEXT("TranslationExportQueryFailedCommandNull", "Translation Export Query Failed: Command is null.");
+					}
+				}
+				else if (ResponseCode == EHttpResponseCodes::NoContent)
+				{
+					bResult = false;
+				}
+				// If there's no response code, then this is the file
+				else
+				{
+					if (Command != nullptr)
+					{
+						TSharedPtr<FOneSkyTranslationExportOperation, ESPMode::ThreadSafe> TranslationExportOp = StaticCastSharedRef<FOneSkyTranslationExportOperation>(Command->Operation);
+						if (TranslationExportOp.IsValid())
+						{
+							FString Filename = TranslationExportOp->GetInExportFileName();
+
+							if (Filename.IsEmpty())
+							{
+								ErrorText = LOCTEXT("InvalidExportFilename", "Export filename is invalid");
+								bResult = false;
+							}
+							else if (FFileHelper::SaveStringToFile(ResponseStr, *Filename, FFileHelper::EEncodingOptions::ForceUnicode))
+							{
+								bResult = true;
+							}
+							else
+							{
+								ErrorText = LOCTEXT("FailedToWriteFile", "Could not write file.");
+								bResult = false;
+							}
+						}
+						else
+						{
+							ErrorText = LOCTEXT("ExportFilenameNotFound", "Could not find export file name.");
+							bResult = false;
+						}
+					}
+				}
+			}
+			else
+			{
+				ErrorText = FText::Format(LOCTEXT("InvalidResponse", "Invalid response. code={0} error={1}"), FText::FromString(FString::FromInt(HttpResponse->GetResponseCode())), FText::FromString(ResponseStr));
+			}
+		}
+
+		if (!bResult)
+		{
+			UE_LOG(LogLocalizationService, Warning, TEXT("%s"), *(ErrorText.ToString()));
+		}
+	}
+
+	if (Command != nullptr)
+	{
+		Command->bCommandSuccessful = bResult;
+		FPlatformAtomics::InterlockedExchange(&(Command->bExecuteProcessed), 1);
+	}
+}
+
+// end TRANSLATION EXPORT
+
+
+// LIST UPLOADED FILES
+
+FName FOneSkyListUploadedFilesWorker::GetName() const
+{
+	return "ListUploadedFilesWorker";
+}
+
+bool FOneSkyListUploadedFilesWorker::Execute(class FOneSkyLocalizationServiceCommand& InCommand)
+{
+	// Store pointer to command so we can access it in the http request callback
+	Command = &InCommand;
+	TSharedPtr<FOneSkyListUploadedFilesOperation, ESPMode::ThreadSafe> ListUploadedFilesOp = StaticCastSharedRef<FOneSkyListUploadedFilesOperation>(InCommand.Operation);
+
+	int32 InProjectId = -1;
+	int32 InStartPage = -1;
+	int32 InItemsPerPage = -1;
+
+	if (ListUploadedFilesOp.IsValid())
+	{
+		InProjectId = ListUploadedFilesOp->GetInProjectId();
+		InStartPage = ListUploadedFilesOp->GetInStartPage();
+		InItemsPerPage = ListUploadedFilesOp->GetInItemsPerPage();
+	}
+
+	FString Url = AddAuthenticationParameters(InCommand.ConnectionInfo, "https://platform.api.onesky.io/1/projects/" + FString::FromInt(InProjectId) / "files");
+	if (InStartPage != -1 && InItemsPerPage != -1)
+	{
+		Url += FString(TEXT("&page=")) + FString::FromInt(InStartPage) + FString(TEXT("&per_page=")) + FString::FromInt(InItemsPerPage);
+	}
+
+	FHttpModule& HttpModule = FModuleManager::LoadModuleChecked<FHttpModule>("HTTP");
+	TSharedRef<class IHttpRequest> HttpRequest = HttpModule.Get().CreateRequest();
+
+	// kick off http request to read
+	HttpRequest->OnProcessRequestComplete().BindRaw(this, &FOneSkyListUploadedFilesWorker::Query_HttpRequestComplete);
+	HttpRequest->SetURL(Url);
+	HttpRequest->SetHeader(TEXT("Content-Type"), TEXT("application/json; charset=utf-8"));
+	HttpRequest->SetVerb(TEXT("GET"));
+	HttpRequest->ProcessRequest();
+
+	return InCommand.bCommandSuccessful;
+}
+
+void FOneSkyListUploadedFilesWorker::Query_HttpRequestComplete(FHttpRequestPtr HttpRequest, FHttpResponsePtr HttpResponse, bool bSucceeded)
+{
+	bool bResult = false;
+
+	if (bSucceeded)
+	{
+		UScriptStruct& StructType = *(FOneSkyListUploadedFilesResponse::StaticStruct());
+		bResult = DeserializeResponseToStruct(&OutListUploadedFilesResponse, StructType, HttpResponse);
+	}
+
+	if (Command != nullptr)
+	{
+		Command->bCommandSuccessful = bResult;
+		FPlatformAtomics::InterlockedExchange(&(Command->bExecuteProcessed), 1);
+	}
+}
+
+// end LIST UPLOADED FILES
+
+
+// UPLOAD FILE
+
+FName FOneSkyUploadFileWorker::GetName() const
+{
+	return "UploadFileWorker";
+}
+
+bool FOneSkyUploadFileWorker::Execute(class FOneSkyLocalizationServiceCommand& InCommand)
+{
+	// Store pointer to command so we can access it in the http request callback
+	Command = &InCommand;
+	TSharedPtr<FOneSkyUploadFileOperation, ESPMode::ThreadSafe> UploadFilesOp = StaticCastSharedRef<FOneSkyUploadFileOperation>(InCommand.Operation);
+
+	int32 InProjectId = -1;
+	FString InFilePathAndName = "";
+	FString InOneSkyTargetFileName = "";
+	FString InLocale = "";
+	bool bInIsKeepingAllStrings = false;
+	FString InFileFormat = "";
+
+	if (UploadFilesOp.IsValid())
+	{
+		InProjectId = UploadFilesOp->GetInProjectId();
+		InFilePathAndName = UploadFilesOp->GetInLocalFilePathAndName();
+		InOneSkyTargetFileName = UploadFilesOp->GetInOneSkyTargetFileName();
+		InLocale = UploadFilesOp->GetInLocale();
+		bInIsKeepingAllStrings = UploadFilesOp->GetInIsKeepingAllStrings();
+		InFileFormat = GetFileFormat(FPaths::GetExtension(InFilePathAndName, true));
+	}
+
+	FString Url = AddAuthenticationParameters(InCommand.ConnectionInfo, "https://platform.api.onesky.io/1/projects/" + FString::FromInt(InProjectId) / "files");
+	Url += "&file_format=" + InFileFormat + "&locale=" + InLocale;
+	if (!bInIsKeepingAllStrings)
+	{
+		Url += "&is_keeping_all_strings=false";
+	}
+
+	FHttpModule& HttpModule = FModuleManager::LoadModuleChecked<FHttpModule>("HTTP");
+	TSharedRef<class IHttpRequest> HttpRequest = HttpModule.Get().CreateRequest();
+
+	// kick off http request to read
+	HttpRequest->OnProcessRequestComplete().BindRaw(this, &FOneSkyUploadFileWorker::Query_HttpRequestComplete);
+	HttpRequest->SetURL(Url);
+	FString Boundary = "---------------------------" + FString::FromInt(FDateTime::Now().GetTicks());
+	HttpRequest->SetHeader(TEXT("Content-Type"), TEXT("multipart/form-data; boundary =" + Boundary));
+	HttpRequest->SetVerb(TEXT("POST"));
+
+	FString FileContents;
+	if (FFileHelper::LoadFileToString(FileContents, *InFilePathAndName, FFileHelper::EEncodingOptions::ForceUnicode))
+	{
+		// Format the file the way OneSky wants it (in a form, with the file contents being a field called "file")
+		FString PrefixBoundry = "\r\n--" + Boundary + "\r\n";
+		FString FileHeader = "Content-Disposition: form-data; name=\"file\"; filename=\"" + InOneSkyTargetFileName + "\"\r\nContent-Type: " + InFileFormat + "\r\n\r\n";
+		FString SuffixBoundary = "\r\n--" + Boundary + "--\r\n";
+		FString ContentsString = PrefixBoundry + FileHeader + FileContents + SuffixBoundary;
+
+		HttpRequest->SetContentAsString(ContentsString);
+		HttpRequest->ProcessRequest();
+	}
+	else
+	{
+		Command->bCommandSuccessful = false;
+		FPlatformAtomics::InterlockedExchange(&(Command->bExecuteProcessed), 1);
+	}
+
+	return InCommand.bCommandSuccessful;
+}
+
+void FOneSkyUploadFileWorker::Query_HttpRequestComplete(FHttpRequestPtr HttpRequest, FHttpResponsePtr HttpResponse, bool bSucceeded)
+{
+	bool bResult = false;
+
+	if (bSucceeded)
+	{
+		UScriptStruct& StructType = *(FOneSkyUploadFileResponse::StaticStruct());
+		bResult = DeserializeResponseToStruct(&OutUploadFileResponse, StructType, HttpResponse);
+	}
+
+	if (Command != nullptr)
+	{
+		Command->bCommandSuccessful = bResult;
+		FPlatformAtomics::InterlockedExchange(&(Command->bExecuteProcessed), 1);
+	}
+}
+
+// end UPLOAD FILE
+
+
+// LIST PHRASE COLLECTIONS
+
+FName FOneSkyListPhraseCollectionsWorker::GetName() const
+{
+	return "ListPhraseCollections";
+}
+
+bool FOneSkyListPhraseCollectionsWorker::Execute(class FOneSkyLocalizationServiceCommand& InCommand)
+{
+	// Store pointer to command so we can access it in the http request callback
+	Command = &InCommand;
+	TSharedPtr<FOneSkyListPhraseCollectionsOperation, ESPMode::ThreadSafe> ListPhraseCollectionsOp = StaticCastSharedRef<FOneSkyListPhraseCollectionsOperation>(InCommand.Operation);
+
+	int32 InProjectId = -1;
+	int32 InPage = -1;
+	int32 InItemsPerPage = -1;
+
+	if (ListPhraseCollectionsOp.IsValid())
+	{
+		InProjectId = ListPhraseCollectionsOp->GetInProjectId();
+		InPage = ListPhraseCollectionsOp->GetInPage();
+		InItemsPerPage = ListPhraseCollectionsOp->GetInItemsPerPage();
+	}
+
+	FString Url = AddAuthenticationParameters(InCommand.ConnectionInfo, "https://platform.api.onesky.io/1/projects/" + FString::FromInt(InProjectId) / "phrase-collections");
+	if (InPage >= 1)
+	{
+		Url += FString(TEXT("&page=")) + FString::FromInt(InPage);
+	}
+	if (InItemsPerPage >= 1)
+	{
+		Url += FString(TEXT("&per_page=")) + FString::FromInt(InItemsPerPage);
+	}
+
+	FHttpModule& HttpModule = FModuleManager::LoadModuleChecked<FHttpModule>("HTTP");
+	TSharedRef<class IHttpRequest> HttpRequest = HttpModule.Get().CreateRequest();
+
+	// kick off http request to read
+	HttpRequest->OnProcessRequestComplete().BindRaw(this, &FOneSkyListPhraseCollectionsWorker::Query_HttpRequestComplete);
+	HttpRequest->SetURL(Url);
+	HttpRequest->SetHeader(TEXT("Content-Type"), TEXT("application/json; charset=utf-8"));
+	HttpRequest->SetVerb(TEXT("GET"));
+	HttpRequest->ProcessRequest();
+
+	return InCommand.bCommandSuccessful;
+}
+
+void FOneSkyListPhraseCollectionsWorker::Query_HttpRequestComplete(FHttpRequestPtr HttpRequest, FHttpResponsePtr HttpResponse, bool bSucceeded)
+{
+	bool bResult = false;
+
+	if (bSucceeded)
+	{
+		UScriptStruct& StructType = *(FOneSkyListPhraseCollectionsResponse::StaticStruct());
+		bResult = DeserializeResponseToStruct(&OutListPhraseCollectionsResponse, StructType, HttpResponse);
+	}
+
+	if (Command != nullptr)
+	{
+		Command->bCommandSuccessful = bResult;
+		FPlatformAtomics::InterlockedExchange(&(Command->bExecuteProcessed), 1);
+	}
+}
+
+// end LIST PHRASE COLLECTIONS
+
+
+//FName FOneSkyListProjectTypesWorker::GetName() const
+//{
+//	return FName(TEXT("FOneSkyListProjectTypesWorker"));
+//}
+//
+//bool FOneSkyListProjectTypesWorker::Execute(class FOneSkyLocalizationServiceCommand& InCommand)
+//{
+//	return false;
+//}
+//
+//FName FOneSkyGetTranslationFileWorker::GetName() const
+//{
+//	return FName(TEXT("FOneSkyGetTranslationFileWorker"));
+//}
+//
+//bool FOneSkyGetTranslationFileWorker::Execute(class FOneSkyLocalizationServiceCommand& InCommand)
+//{
+//	return false;
+//}
+//
+//FName FOneSkyShowPhraseCollectionWorker::GetName() const
+//{
+//	return FName(TEXT("FOneSkyShowPhraseCollectionWorker"));
+//}
+//
+//bool FOneSkyShowPhraseCollectionWorker::Execute(class FOneSkyLocalizationServiceCommand& InCommand)
+//{
+//	UpdateCachedLocalizationStates(OutStates);
+//	return false;
+//}
+//
+//FName FOneSkyImportPhraseCollectionWorker::GetName() const
+//{
+//	return FName(TEXT("FOneSkyShowPhraseCollectionWorker"));
+//}
+//
+//bool FOneSkyImportPhraseCollectionWorker::Execute(class FOneSkyLocalizationServiceCommand& InCommand)
+//{
+//	return false;
+//}
+
+
+#undef LOCTEXT_NAMESPACE
+
diff --git a/Engine/Plugins/Developer/OneSkyLocalizationService/Source/OneSkyLocalizationService/Private/OneSkyLocalizationServiceOperations.h b/Engine/Plugins/Developer/OneSkyLocalizationService/Source/OneSkyLocalizationService/Private/OneSkyLocalizationServiceOperations.h
new file mode 100644
--- /dev/null
+++ b/Engine/Plugins/Developer/OneSkyLocalizationService/Source/OneSkyLocalizationService/Private/OneSkyLocalizationServiceOperations.h
@@ -0,0 +1,1014 @@
+// Copyright 1998-2015 Epic Games, Inc. All Rights Reserved.
+
+#pragma once
+
+#include "IOneSkyLocalizationServiceWorker.h"
+#include "OneSkyLocalizationServiceState.h"
+#include "OneSkyLocalizationServiceResponseTypes.h"
+
+#define LOCTEXT_NAMESPACE "LocalizationService"
+typedef TSharedPtr<class IHttpRequest> FHttpRequestPtr;
+typedef TSharedPtr<class IHttpResponse, ESPMode::ThreadSafe> FHttpResponsePtr;
+
+class FOneSkyConnect : public IOneSkyLocalizationServiceWorker
+{
+public:
+	virtual ~FOneSkyConnect() {}
+	// IOneSkyLocalizationServiceWorker interface
+	virtual FName GetName() const override;
+	virtual bool Execute(class FOneSkyLocalizationServiceCommand& InCommand) override;
+	virtual bool UpdateStates() const override {/*TODO:*/return true; }
+public:
+
+};
+
+
+// LIST PROJECT GROUPS
+
+/**
+* Operation used to list project groups in one sky
+*/
+class FOneSkyListProjectGroupsOperation : public ILocalizationServiceOperation
+{
+public:
+	// ILocalizationServiceOperation interface
+	virtual FName GetName() const override
+	{
+		return "ListProjectGroups";
+	}
+
+	virtual FText GetInProgressString() const override
+	{
+		return LOCTEXT("OneSkyLocalizationService_ListProjectGroupsWorker", "Listing Project Groups in OneSky...");
+	}
+
+	FOneSkyListProjectGroupsOperation() : InStartPage(-1), InItemsPerPage(-1) { }
+
+	void SetInStartPage(int32 NewStartPage)
+	{
+		InStartPage = NewStartPage;
+	}
+
+	int32 GetInStartPage()
+	{
+		return InStartPage;
+	}
+
+	void SetInItemsPerPage(int32 NewItemsPerPage)
+	{
+		InItemsPerPage = NewItemsPerPage;
+	}
+
+	int32 GetInItemsPerPage()
+	{
+		return InItemsPerPage;
+	}
+
+protected:
+	
+	int32 InStartPage;
+	int32 InItemsPerPage;
+};
+
+class FOneSkyListProjectGroupsWorker : public IOneSkyLocalizationServiceWorker
+{
+public:
+	virtual ~FOneSkyListProjectGroupsWorker() {}
+	// IOneSkyLocalizationServiceWorker interface
+	virtual FName GetName() const override;
+	virtual bool Execute(class FOneSkyLocalizationServiceCommand& InCommand) override;
+	virtual bool UpdateStates() const override {/*TODO:*/return true; }
+	void Query_HttpRequestComplete(FHttpRequestPtr HttpRequest, FHttpResponsePtr HttpResponse, bool bSucceeded);
+public:
+
+	FOneSkyLocalizationServiceCommand* Command;
+
+	/** Response to the  command */
+	FOneSkyListProjectGroupsResponse OutListProjectGroupsResponse;
+};
+
+// end LIST PROJECT GROUPS
+
+
+
+// SHOW PROJECT GROUP
+
+/**
+* Operation used to show a project group in one sky
+*/
+class FOneSkyShowProjectGroupOperation : public ILocalizationServiceOperation
+{
+public:
+	// ILocalizationServiceOperation interface
+	virtual FName GetName() const override
+	{
+		return "ShowProjectGroup";
+	}
+
+	virtual FText GetInProgressString() const override
+	{
+		return LOCTEXT("OneSkyLocalizationService_ShowProjectGroupWorker", "Showing Project Group in OneSky...");
+	}
+
+	FOneSkyShowProjectGroupOperation() : InProjectGroupId(-1) { }
+
+	int32 GetInProjectGroupId()
+	{
+		return InProjectGroupId;
+	}
+
+	void SetInProjectGroupId(int32 NewProjectGroupId)
+	{
+		InProjectGroupId = NewProjectGroupId;
+	}
+
+protected:
+
+	int32 InProjectGroupId;
+};
+
+class FOneSkyShowProjectGroupWorker : public IOneSkyLocalizationServiceWorker
+{
+public:
+	virtual ~FOneSkyShowProjectGroupWorker() {}
+	// IOneSkyLocalizationServiceWorker interface
+	virtual FName GetName() const override;
+	virtual bool Execute(class FOneSkyLocalizationServiceCommand& InCommand) override;
+	virtual bool UpdateStates() const override {/*TODO:*/return true; }
+	void Query_HttpRequestComplete(FHttpRequestPtr HttpRequest, FHttpResponsePtr HttpResponse, bool bSucceeded);
+public:
+
+	FOneSkyLocalizationServiceCommand* Command;
+
+	/** Response to the command */
+	FOneSkyShowProjectGroupResponse OutShowProjectGroupResponse;
+};
+
+// end SHOW PROJECT GROUP
+
+
+// CREATE PROJECT GROUP
+
+/**
+* Operation used to create a project group in one sky
+*/
+class FOneSkyCreateProjectGroupOperation : public ILocalizationServiceOperation
+{
+public:
+	// ILocalizationServiceOperation interface
+	virtual FName GetName() const override
+	{
+		return "CreateProjectGroup";
+	}
+
+	virtual FText GetInProgressString() const override
+	{
+		return LOCTEXT("OneSkyLocalizationService_CreateProjectGroupWorker", "Creating Project Group in OneSky...");
+	}
+
+	FOneSkyCreateProjectGroupOperation() : InProjectGroupName(""), InProjectGroupBaseLocale("") {}
+
+	FString GetInProjectGroupName()
+	{
+		return InProjectGroupName;
+	}
+
+	void SetInProjectGroupName(FString NewProjectGroupName)
+	{
+		InProjectGroupName = NewProjectGroupName;
+	}
+
+	FString GetInProjectGroupBaseLocale()
+	{
+		return InProjectGroupBaseLocale;
+	}
+
+	void SetInInProjectGroupBaseLocale(FString NewProjectGroupBaseLocale)
+	{
+		InProjectGroupBaseLocale = NewProjectGroupBaseLocale;
+	}
+
+protected:
+
+	FString InProjectGroupName;
+	FString InProjectGroupBaseLocale;
+
+};
+
+class FOneSkyCreateProjectGroupWorker : public IOneSkyLocalizationServiceWorker
+{
+public:
+	virtual ~FOneSkyCreateProjectGroupWorker() {}
+	// IOneSkyLocalizationServiceWorker interface
+	virtual FName GetName() const override;
+	virtual bool Execute(class FOneSkyLocalizationServiceCommand& InCommand) override;
+	virtual bool UpdateStates() const override {/*TODO:*/return true; }
+	void Query_HttpRequestComplete(FHttpRequestPtr HttpRequest, FHttpResponsePtr HttpResponse, bool bSucceeded);
+public:
+
+	FOneSkyLocalizationServiceCommand* Command;
+
+	/** Response to the command */
+	FOneSkyCreateProjectGroupResponse OutCreateProjectGroupResponse;
+};
+
+// end CREATE PROJECT GROUP
+
+
+// LIST PROJECT GROUP LANGUAGES
+
+/**
+* Operation used to list the languages supported by a project group in one sky
+*/
+class FOneSkyListProjectGroupLanguagesOperation : public ILocalizationServiceOperation
+{
+public:
+	// ILocalizationServiceOperation interface
+	virtual FName GetName() const override
+	{
+		return "ListProjectGroupLanguages";
+	}
+
+	virtual FText GetInProgressString() const override
+	{
+		return LOCTEXT("OneSkyLocalizationService_ListProjectGroupLanguagesWorker", "Listing languages for Project Group in OneSky...");
+	}
+
+	FOneSkyListProjectGroupLanguagesOperation() : InProjectGroupId(-1) { }
+
+	int32 GetInProjectGroupId()
+	{
+		return InProjectGroupId;
+	}
+
+	void SetInProjectGroupId(int32 NewProjectGroupId)
+	{
+		InProjectGroupId = NewProjectGroupId;
+	}
+
+protected:
+
+	int32 InProjectGroupId;
+};
+
+class FOneSkyListProjectGroupLanguagesWorker : public IOneSkyLocalizationServiceWorker
+{
+public:
+	virtual ~FOneSkyListProjectGroupLanguagesWorker() {}
+	// IOneSkyLocalizationServiceWorker interface
+	virtual FName GetName() const override;
+	virtual bool Execute(class FOneSkyLocalizationServiceCommand& InCommand) override;
+	virtual bool UpdateStates() const override {/*TODO:*/return true; }
+	void Query_HttpRequestComplete(FHttpRequestPtr HttpRequest, FHttpResponsePtr HttpResponse, bool bSucceeded);
+public:
+
+	FOneSkyLocalizationServiceCommand* Command;
+
+	/** Response to the command */
+	FOneSkyListProjectGroupLanguagesResponse OutListProjectGroupLanguagesResponse;
+};
+
+// end LIST PROJECT GROUP LANGUAGES
+
+
+// LIST PROJECTS IN GROUP
+
+/**
+* Operation used to list projects in a project group in one sky
+*/
+class FOneSkyListProjectsInGroupOperation : public ILocalizationServiceOperation
+{
+public:
+	// ILocalizationServiceOperation interface
+	virtual FName GetName() const override
+	{
+		return "ListProjectsInGroup";
+	}
+
+	virtual FText GetInProgressString() const override
+	{
+		return LOCTEXT("OneSkyLocalizationService_ListProjectsInGroupWorker", "Listing Projects in Groups in OneSky...");
+	}
+
+	FOneSkyListProjectsInGroupOperation() : InProjectGroupId(-1) { }
+
+	int32 GetInProjectGroupId()
+	{
+		return InProjectGroupId;
+	}
+
+	void SetInProjectGroupId(int32 NewProjectGroupId)
+	{
+		InProjectGroupId = NewProjectGroupId;
+	}
+
+protected:
+
+	int32 InProjectGroupId;
+};
+
+class FOneSkyListProjectsInGroupWorker : public IOneSkyLocalizationServiceWorker
+{
+public:
+	virtual ~FOneSkyListProjectsInGroupWorker() {}
+	// IOneSkyLocalizationServiceWorker interface
+	virtual FName GetName() const override;
+	virtual bool Execute(class FOneSkyLocalizationServiceCommand& InCommand) override;
+	virtual bool UpdateStates() const override {/*TODO:*/return true; }
+	void Query_HttpRequestComplete(FHttpRequestPtr HttpRequest, FHttpResponsePtr HttpResponse, bool bSucceeded);
+public:
+
+	FOneSkyLocalizationServiceCommand* Command;
+
+	/** Response to the get project groups command */
+	FOneSkyListProjectsInGroupResponse OutListProjectsInGroupResponse;
+};
+
+// end LIST PROJECTS IN GROUP
+
+// SHOW PROJECT
+
+/**
+* Operation used to show a project in one sky
+*/
+class FOneSkyShowProjectOperation : public ILocalizationServiceOperation
+{
+public:
+	// ILocalizationServiceOperation interface
+	virtual FName GetName() const override
+	{
+		return "ShowProject";
+	}
+
+	virtual FText GetInProgressString() const override
+	{
+		return LOCTEXT("OneSkyLocalizationService_ShowProjectWorker", "Showing Project in OneSky...");
+	}
+
+	FOneSkyShowProjectOperation() : InProjectId(-1) { }
+
+	int32 GetInProjectId()
+	{
+		return InProjectId;
+	}
+
+	void SetInProjectId(int32 NewProjectId)
+	{
+		InProjectId = NewProjectId;
+	}
+
+protected:
+
+	int32 InProjectId;
+};
+
+class FOneSkyShowProjectWorker : public IOneSkyLocalizationServiceWorker
+{
+public:
+	virtual ~FOneSkyShowProjectWorker() {}
+	// IOneSkyLocalizationServiceWorker interface
+	virtual FName GetName() const override;
+	virtual bool Execute(class FOneSkyLocalizationServiceCommand& InCommand) override;
+	virtual bool UpdateStates() const override {/*TODO:*/return true; }
+	void Query_HttpRequestComplete(FHttpRequestPtr HttpRequest, FHttpResponsePtr HttpResponse, bool bSucceeded);
+public:
+
+	FOneSkyLocalizationServiceCommand* Command;
+
+	/** Response to the command */
+	FOneSkyShowProjectResponse OutShowProjectResponse;
+};
+
+// end SHOW PROJECT
+
+
+
+// CREATE PROJECT
+
+/**
+* Operation used to create a project in one sky
+*/
+class FOneSkyCreateProjectOperation : public ILocalizationServiceOperation
+{
+public:
+	// ILocalizationServiceOperation interface
+	virtual FName GetName() const override
+	{
+		return "CreateProject";
+	}
+
+	virtual FText GetInProgressString() const override
+	{
+		return LOCTEXT("OneSkyLocalizationService_CreateProjectWorker", "Creating Project in OneSky...");
+	}
+
+	FOneSkyCreateProjectOperation() : InProjectGroupId(-1), InProjectType(""), InProjectName(""), InProjectDescription("") {}
+
+	FString GetInProjectGroupName()
+	{
+		return InProjectName;
+	}
+
+	void SetInProjectGroupName(FString NewProjectName)
+	{
+		InProjectName = NewProjectName;
+	}
+
+	FString GetInProjectDescription()
+	{
+		return InProjectDescription;
+	}
+
+	void SetInInProjectDescription(FString NewProjectDescription)
+	{
+		InProjectDescription = NewProjectDescription;
+	}
+
+	FString GetInProjectType()
+	{
+		return InProjectType;
+	}
+
+	void SetInInProjectType(FString NewProjectType)
+	{
+		InProjectType = NewProjectType;
+	}
+
+	int32 GetInProjectGroupId()
+	{
+		return InProjectGroupId;
+	}
+
+	void SetInProjectGroupId(int32 NewProjectGroupId)
+	{
+		InProjectGroupId = NewProjectGroupId;
+	}
+
+protected:
+
+	int32 InProjectGroupId;	// Project Group this project will belong to
+	FString InProjectType; // See [REDACTED_PROFILE_URL]/api-documentation-platform/blob/master/resources/project_type.md
+	FString InProjectName;
+	FString InProjectDescription;
+
+};
+
+class FOneSkyCreateProjectWorker : public IOneSkyLocalizationServiceWorker
+{
+public:
+	virtual ~FOneSkyCreateProjectWorker() {}
+	// IOneSkyLocalizationServiceWorker interface
+	virtual FName GetName() const override;
+	virtual bool Execute(class FOneSkyLocalizationServiceCommand& InCommand) override;
+	virtual bool UpdateStates() const override {/*TODO:*/return true; }
+	void Query_HttpRequestComplete(FHttpRequestPtr HttpRequest, FHttpResponsePtr HttpResponse, bool bSucceeded);
+public:
+
+	FOneSkyLocalizationServiceCommand* Command;
+
+	/** Response to the command */
+	FOneSkyCreateProjectResponse OutCreateProjectResponse;
+};
+
+// end CREATE PROJECT
+
+
+
+// LIST PROJECT GROUP LANGUAGES
+
+/**
+* Operation used to list the languages supported by a project in one sky
+*/
+class FOneSkyListProjectLanguagesOperation : public ILocalizationServiceOperation
+{
+public:
+	// ILocalizationServiceOperation interface
+	virtual FName GetName() const override
+	{
+		return "ListProjectLanguages";
+	}
+
+	virtual FText GetInProgressString() const override
+	{
+		return LOCTEXT("OneSkyLocalizationService_ListProjectGroupLanguagesWorker", "Listing languages for Project in OneSky...");
+	}
+
+	FOneSkyListProjectLanguagesOperation() : InProjectId(-1) { }
+
+	int32 GetInProjectId()
+	{
+		return InProjectId;
+	}
+
+	void SetInProjectId(int32 NewProjectGroupId)
+	{
+		InProjectId = NewProjectGroupId;
+	}
+
+protected:
+
+	int32 InProjectId;
+};
+
+class FOneSkyListProjectLanguagesWorker : public IOneSkyLocalizationServiceWorker
+{
+public:
+	virtual ~FOneSkyListProjectLanguagesWorker() {}
+	// IOneSkyLocalizationServiceWorker interface
+	virtual FName GetName() const override;
+	virtual bool Execute(class FOneSkyLocalizationServiceCommand& InCommand) override;
+	virtual bool UpdateStates() const override {/*TODO:*/return true; }
+	void Query_HttpRequestComplete(FHttpRequestPtr HttpRequest, FHttpResponsePtr HttpResponse, bool bSucceeded);
+public:
+
+	FOneSkyLocalizationServiceCommand* Command;
+
+	/** Response to the command */
+	FOneSkyListProjectLanguagesResponse OutListProjectLanguagesResponse;
+};
+
+// end LIST PROJECT GROUP LANGUAGES
+
+
+// TRANSLATION STATUS
+
+/**
+* Operation used to list the languages supported by a project in one sky
+*/
+class FOneSkyTranslationStatusOperation : public ILocalizationServiceOperation
+{
+public:
+	// ILocalizationServiceOperation interface
+	virtual FName GetName() const override
+	{
+		return "TranslationStatus";
+	}
+
+	virtual FText GetInProgressString() const override
+	{
+		return LOCTEXT("OneSkyLocalizationService_TranslationStatus", "Retrieving the translation status for a file in a project in OneSky...");
+	}
+
+	FOneSkyTranslationStatusOperation() : InProjectId(-1) { }
+
+	int32 GetInProjectId()
+	{
+		return InProjectId;
+	}
+
+	void SetInProjectId(int32 NewProjectGroupId)
+	{
+		InProjectId = NewProjectGroupId;
+	}
+
+	FString GetInFileName()
+	{
+		return InFileName;
+	}
+
+	void SetInFileName(FString NewFileName)
+	{
+		InFileName = NewFileName;
+	}
+
+	FString GetInLocale()
+	{
+		return InLocale;
+	}
+
+	void SetInLocale(FString NewLocale)
+	{
+		InLocale = NewLocale;
+	}
+
+	FString GetOutPercentComplete()
+	{
+		return OutPercentComplete;
+	}
+
+	void SetOutPercentComplete(FString NewOutPercentComplete)
+	{
+		OutPercentComplete = NewOutPercentComplete;
+	}
+
+protected:
+
+	int32 InProjectId;
+	FString InFileName;
+	FString InLocale;
+	FString OutPercentComplete;
+};
+
+class FOneSkyTranslationStatusWorker : public IOneSkyLocalizationServiceWorker
+{
+public:
+	virtual ~FOneSkyTranslationStatusWorker() {}
+	// IOneSkyLocalizationServiceWorker interface
+	virtual FName GetName() const override;
+	virtual bool Execute(class FOneSkyLocalizationServiceCommand& InCommand) override;
+	virtual bool UpdateStates() const override {/*TODO:*/return true; }
+	void Query_HttpRequestComplete(FHttpRequestPtr HttpRequest, FHttpResponsePtr HttpResponse, bool bSucceeded);
+public:
+
+	FOneSkyLocalizationServiceCommand* Command;
+
+	/** Response to the command */
+	FOneSkyTranslationStatusResponse OutTranslationStatusResponse;
+};
+
+// end TRANSLATION STATUS
+
+
+// TRANSLATION EXPORT
+
+/**
+* Operation used to list the languages supported by a project in one sky
+*/
+class FOneSkyTranslationExportOperation : public ILocalizationServiceOperation
+{
+public:
+	// ILocalizationServiceOperation interface
+	virtual FName GetName() const override
+	{
+		return "TranslationExport";
+	}
+
+	virtual FText GetInProgressString() const override
+	{
+		return LOCTEXT("OneSkyLocalizationService_TranslationExport", "Exporting translations for a file in a project in OneSky...");
+	}
+
+	FOneSkyTranslationExportOperation() : InProjectId(-1) { }
+
+	int32 GetInProjectId()
+	{
+		return InProjectId;
+	}
+
+	void SetInProjectId(int32 NewProjectGroupId)
+	{
+		InProjectId = NewProjectGroupId;
+	}
+
+	FString GetInSourceFileName()
+	{
+		return InSourceFileName;
+	}
+
+	void SetInSourceFileName(FString NewSourceFileName)
+	{
+		InSourceFileName = NewSourceFileName;
+	}
+
+	FString GetInExportFileName()
+	{
+		return InExportFileName;
+	}
+
+	void SetInExportFileName(FString NewExportFileName)
+	{
+		InExportFileName = NewExportFileName;
+	}
+
+	FString GetInLocale()
+	{
+		return InLocale;
+	}
+
+	void SetInLocale(FString NewLocale)
+	{
+		InLocale = NewLocale;
+	}
+
+protected:
+
+	int32 InProjectId;
+	FString InSourceFileName;
+	FString InExportFileName;
+	FString InLocale;
+};
+
+class FOneSkyTranslationExportWorker : public IOneSkyLocalizationServiceWorker
+{
+public:
+	virtual ~FOneSkyTranslationExportWorker() {}
+	// IOneSkyLocalizationServiceWorker interface
+	virtual FName GetName() const override;
+	virtual bool Execute(class FOneSkyLocalizationServiceCommand& InCommand) override;
+	virtual bool UpdateStates() const override {/*TODO:*/return true; }
+	void Query_HttpRequestComplete(FHttpRequestPtr HttpRequest, FHttpResponsePtr HttpResponse, bool bSucceeded);
+public:
+
+	FOneSkyLocalizationServiceCommand* Command;
+
+	/** Response to the command */
+	//FOneSkyTranslationExportResponse OutTranslationExportResponse;
+};
+
+// end TRANSLATION EXPORT
+
+
+// LIST UPLOADED FILES
+
+/**
+* Operation used to list uploaded files for a project in one sky
+*/
+class FOneSkyListUploadedFilesOperation : public ILocalizationServiceOperation
+{
+public:
+	// ILocalizationServiceOperation interface
+	virtual FName GetName() const override
+	{
+		return "ListUploadedFiles";
+	}
+
+	virtual FText GetInProgressString() const override
+	{
+		return LOCTEXT("OneSkyLocalizationService_ListUploadedFilesWorker", "Listing Uploaded Files for a Project in OneSky...");
+	}
+
+	FOneSkyListUploadedFilesOperation() : InProjectId(-1), InStartPage(-1), InItemsPerPage(-1) { }
+
+	void SetInStartPage(int32 NewStartPage)
+	{
+		InStartPage = NewStartPage;
+	}
+
+	int32 GetInStartPage()
+	{
+		return InStartPage;
+	}
+
+	void SetInItemsPerPage(int32 NewItemsPerPage)
+	{
+		InItemsPerPage = NewItemsPerPage;
+	}
+
+	int32 GetInItemsPerPage()
+	{
+		return InItemsPerPage;
+	}
+
+	int32 GetInProjectId()
+	{
+		return InProjectId;
+	}
+
+	void SetInProjectId(int32 NewProjectGroupId)
+	{
+		InProjectId = NewProjectGroupId;
+	}
+
+protected:
+
+	int32 InProjectId;
+	int32 InStartPage;
+	int32 InItemsPerPage;
+};
+
+class FOneSkyListUploadedFilesWorker : public IOneSkyLocalizationServiceWorker
+{
+public:
+	virtual ~FOneSkyListUploadedFilesWorker() {}
+	// IOneSkyLocalizationServiceWorker interface
+	virtual FName GetName() const override;
+	virtual bool Execute(class FOneSkyLocalizationServiceCommand& InCommand) override;
+	virtual bool UpdateStates() const override {/*TODO:*/return true; }
+	void Query_HttpRequestComplete(FHttpRequestPtr HttpRequest, FHttpResponsePtr HttpResponse, bool bSucceeded);
+public:
+
+	FOneSkyLocalizationServiceCommand* Command;
+
+	/** Response to the  command */
+	FOneSkyListUploadedFilesResponse OutListUploadedFilesResponse;
+};
+
+// end LIST UPLOADED FILES
+
+
+// UPLOADED FILE
+
+/**
+* Operation used to upload a files to a project in one sky
+*/
+class FOneSkyUploadFileOperation : public ILocalizationServiceOperation
+{
+public:
+
+	// ILocalizationServiceOperation interface
+	virtual FName GetName() const override
+	{
+		return "UploadFile";
+	}
+
+	virtual FText GetInProgressString() const override
+	{
+		return LOCTEXT("OneSkyLocalizationService_UploadFileWorker", "Uploading File to a Project in OneSky...");
+	}
+
+	FOneSkyUploadFileOperation() : InProjectId(-1), bInIsKeepingAllStrings(true) { }
+
+	void SetInLocalFilePathAndName(FString NewInLocalFilePathAndName)
+	{
+		InLocalFilePathAndName = NewInLocalFilePathAndName;
+	}
+
+	FString GetInLocalFilePathAndName()
+	{
+		return InLocalFilePathAndName;
+	}
+
+	void SetInOneSkyTargetFileName(FString NewInOneSkyTargetFileName)
+	{
+		InOneSkyTargetFileName = NewInOneSkyTargetFileName;
+	}
+
+	FString GetInOneSkyTargetFileName()
+	{
+		return InOneSkyTargetFileName;
+	}
+
+	void SetInLocale(FString NewInLocale)
+	{
+		InLocale = NewInLocale;
+	}
+
+	FString GetInLocale()
+	{
+		return InLocale;
+	}
+
+	bool GetInIsKeepingAllStrings()
+	{
+		return bInIsKeepingAllStrings;
+	}
+
+	void SetInIsKeepingAllStrings(bool NewInIsKeepingAllStrings)
+	{
+		bInIsKeepingAllStrings = NewInIsKeepingAllStrings;
+	}
+
+	int32 GetInProjectId()
+	{
+		return InProjectId;
+	}
+
+	void SetInProjectId(int32 NewProjectGroupId)
+	{
+		InProjectId = NewProjectGroupId;
+	}
+
+protected:
+
+	int32 InProjectId;
+	FString InLocalFilePathAndName;
+	FString InOneSkyTargetFileName;
+	FString InLocale;
+	bool bInIsKeepingAllStrings;
+};
+
+class FOneSkyUploadFileWorker : public IOneSkyLocalizationServiceWorker
+{
+public:
+	virtual ~FOneSkyUploadFileWorker() {}
+	// IOneSkyLocalizationServiceWorker interface
+	virtual FName GetName() const override;
+	virtual bool Execute(class FOneSkyLocalizationServiceCommand& InCommand) override;
+	virtual bool UpdateStates() const override {/*TODO:*/return true; }
+	void Query_HttpRequestComplete(FHttpRequestPtr HttpRequest, FHttpResponsePtr HttpResponse, bool bSucceeded);
+public:
+
+	FOneSkyLocalizationServiceCommand* Command;
+
+	/** Response to the  command */
+	FOneSkyUploadFileResponse OutUploadFileResponse;
+};
+
+// end UPLOAD FILES
+
+
+// LIST PHRASE COLLECTIONS
+
+/**
+* Operation used to list phrase collections for a project in one sky
+*/
+class FOneSkyListPhraseCollectionsOperation : public ILocalizationServiceOperation
+{
+public:
+	// ILocalizationServiceOperation interface
+	virtual FName GetName() const override
+	{
+		return "ListPhraseCollections";
+	}
+
+	virtual FText GetInProgressString() const override
+	{
+		return LOCTEXT("OneSkyLocalizationService_ListPhraseCollections", "Listing Phrase Collections for a Project in OneSky...");
+	}
+
+	FOneSkyListPhraseCollectionsOperation() : InProjectId(-1), InPage(-1), InItemsPerPage(-1) { }
+
+	void SetInPage(int32 NewPage)
+	{
+		InPage = NewPage;
+	}
+
+	int32 GetInPage()
+	{
+		return InPage;
+	}
+
+	void SetInItemsPerPage(int32 NewItemsPerPage)
+	{
+		InItemsPerPage = NewItemsPerPage;
+	}
+
+	int32 GetInItemsPerPage()
+	{
+		return InItemsPerPage;
+	}
+
+	int32 GetInProjectId()
+	{
+		return InProjectId;
+	}
+
+	void SetInProjectId(int32 NewProjectGroupId)
+	{
+		InProjectId = NewProjectGroupId;
+	}
+
+protected:
+
+	int32 InProjectId;
+	int32 InPage;
+	int32 InItemsPerPage;
+};
+
+class FOneSkyListPhraseCollectionsWorker : public IOneSkyLocalizationServiceWorker
+{
+public:
+	virtual ~FOneSkyListPhraseCollectionsWorker() {}
+	// IOneSkyLocalizationServiceWorker interface
+	virtual FName GetName() const override;
+	virtual bool Execute(class FOneSkyLocalizationServiceCommand& InCommand) override;
+	virtual bool UpdateStates() const override {/*TODO:*/return true; }
+	void Query_HttpRequestComplete(FHttpRequestPtr HttpRequest, FHttpResponsePtr HttpResponse, bool bSucceeded);
+public:
+
+	FOneSkyLocalizationServiceCommand* Command;
+
+	/** Response to the  command */
+	FOneSkyListPhraseCollectionsResponse OutListPhraseCollectionsResponse;
+};
+
+// end LIST PHRASE COLLECTIONS
+
+//class FOneSkyListProjectTypesWorker : public IOneSkyLocalizationServiceWorker
+//{
+//public:
+//	virtual ~FOneSkyListProjectTypesWorker() {}
+//	// IOneSkyLocalizationServiceWorker interface
+//	virtual FName GetName() const override;
+//	virtual bool Execute(class FOneSkyLocalizationServiceCommand& InCommand) override;
+//	virtual bool UpdateStates() const override {/*TODO:*/return true; }
+//public:
+//	/** Map of filenames to perforce state */
+//	//FOneSkeyProjectTypeResponse OutResults;
+//};
+//
+//class FOneSkyShowPhraseCollectionWorker : public IOneSkyLocalizationServiceWorker
+//{
+//public:
+//	virtual ~FOneSkyShowPhraseCollectionWorker() 
+//	{}
+//	// IOneSkyLocalizationServiceWorker interface
+//	virtual FName GetName() const override;
+//	virtual bool Execute(class FOneSkyLocalizationServiceCommand& InCommand) override;
+//	virtual bool UpdateStates() const override {/*TODO:*/return true; }
+//public:
+//	int32 InCollectionKey;
+//
+//	TMap<FLocalizationServiceTranslationIdentifier, TSharedRef<FOneSkyLocalizationServiceState, ESPMode::ThreadSafe>, FDefaultSetAllocator, FLocalizationServiceTranslationIdentifierKeyFuncs<TSharedRef<FOneSkyLocalizationServiceState, ESPMode::ThreadSafe>>> OutStates;
+//	//FOneSkyShowPhraseCollectionResponse ShowPhraseCollectionResponse;
+//};
+//
+//class FOneSkyImportPhraseCollectionWorker : public IOneSkyLocalizationServiceWorker
+//{
+//public:
+//	virtual ~FOneSkyImportPhraseCollectionWorker() {}
+//	// IOneSkyLocalizationServiceWorker interface
+//	virtual FName GetName() const override;
+//	virtual bool Execute(class FOneSkyLocalizationServiceCommand& InCommand) override;
+//	virtual bool UpdateStates() const override {/*TODO:*/return true; }
+//public:
+//	//FOneSkyPhraseCollection InCollection;
+//
+//	// Whether to deprecate strings within an previously created phrase collection that cannot be found in the re - creation of the phrase collection.
+//	bool InKeepAllStrings;	
+//
+//	/** Map of filenames to perforce state */
+//	//FOneSkyImportPhraseCollectionResponse
+//};
+
+#undef LOCTEXT_NAMESPACE
\ No newline at end of file
diff --git a/Engine/Plugins/Developer/OneSkyLocalizationService/Source/OneSkyLocalizationService/Private/OneSkyLocalizationServicePrivatePCH.h b/Engine/Plugins/Developer/OneSkyLocalizationService/Source/OneSkyLocalizationService/Private/OneSkyLocalizationServicePrivatePCH.h
new file mode 100644
--- /dev/null
+++ b/Engine/Plugins/Developer/OneSkyLocalizationService/Source/OneSkyLocalizationService/Private/OneSkyLocalizationServicePrivatePCH.h
@@ -0,0 +1,17 @@
+// Copyright 1998-2015 Epic Games, Inc. All Rights Reserved.
+
+#pragma once
+
+#include "Core.h"
+#include "SlateBasics.h"
+#include "EditorStyle.h"
+
+#include "Developer/LocalizationService/Public/ILocalizationServiceModule.h"
+
+#include "OneSkyLocalizationServiceModule.h"
+#include "OneSkyLocalizationServiceSettings.h"
+#include "OneSkyLocalizationServiceProvider.h"
+#include "OneSkyLocalizationServiceRevision.h"
+#include "OneSkyLocalizationServiceCommand.h"
+#include "OneSkyLocalizationServiceState.h"
+#include "OneSkyLocalizationServiceOperations.h"
\ No newline at end of file
diff --git a/Engine/Plugins/Developer/OneSkyLocalizationService/Source/OneSkyLocalizationService/Private/OneSkyLocalizationServiceProvider.cpp b/Engine/Plugins/Developer/OneSkyLocalizationService/Source/OneSkyLocalizationService/Private/OneSkyLocalizationServiceProvider.cpp
new file mode 100644
--- /dev/null
+++ b/Engine/Plugins/Developer/OneSkyLocalizationService/Source/OneSkyLocalizationService/Private/OneSkyLocalizationServiceProvider.cpp
@@ -0,0 +1,622 @@
+// Copyright 1998-2015 Epic Games, Inc. All Rights Reserved.
+
+#include "OneSkyLocalizationServicePrivatePCH.h"
+#include "OneSkyLocalizationServiceProvider.h"
+#include "OneSkyLocalizationServiceCommand.h"
+#include "OneSkyConnection.h"
+#include "IOneSkyLocalizationServiceWorker.h"
+#include "OneSkyLocalizationServiceModule.h"
+#include "OneSkyLocalizationServiceSettings.h"
+#include "MessageLog.h"
+#include "Developer/LocalizationService/Public/ScopedLocalizationServiceProgress.h"
+#if LOCALIZATION_SERVICES_WITH_SLATE
+
+#include "Editor/PropertyEditor/Public/DetailLayoutBuilder.h"
+#include "Editor/PropertyEditor/Public/DetailWidgetRow.h"
+#include "Editor/PropertyEditor/Public/IDetailPropertyRow.h"
+#include "Editor/PropertyEditor/Public/DetailCategoryBuilder.h"
+#include "Editor/PropertyEditor/Public/IDetailsView.h"
+
+#endif
+
+static FName ProviderName("OneSky");
+
+#define LOCTEXT_NAMESPACE "OneSkyLocalizationService"
+
+/** Init of connection with source control server */
+void FOneSkyLocalizationServiceProvider::Init(bool bForceConnection)
+{
+	// TODO: Test if connection works
+	bServerAvailable = true;
+
+	// TODO: this is called on startup, delay until we need it
+	// Move to EstablishConnection?
+	RetrieveProjectGroups(ELocalizationServiceOperationConcurrency::Asynchronous);
+}
+
+/** API Specific close the connection with localization provider server*/
+void FOneSkyLocalizationServiceProvider::Close()
+{
+	if ( PersistentConnection )
+	{
+		PersistentConnection->Disconnect();
+		delete PersistentConnection;
+		PersistentConnection = NULL;
+	}
+
+	// clear the cache
+	StateCache.Empty();
+	ProjectGroupCache.Empty();
+	ProjectCachePerGroupId.Empty();
+
+	bServerAvailable = false;
+}
+
+TSharedRef<FOneSkyLocalizationServiceState, ESPMode::ThreadSafe> FOneSkyLocalizationServiceProvider::GetStateInternal(const FLocalizationServiceTranslationIdentifier& InTranslationId)
+{
+	TSharedRef<FOneSkyLocalizationServiceState, ESPMode::ThreadSafe>* State = StateCache.Find(InTranslationId);
+	if(State != NULL)
+	{
+		// found cached item
+		return (*State);
+	}
+	else
+	{
+		// cache an unknown state for this item
+		TSharedRef<FOneSkyLocalizationServiceState, ESPMode::ThreadSafe> NewState = MakeShareable(new FOneSkyLocalizationServiceState(InTranslationId));
+		StateCache.Add(InTranslationId, NewState);
+		return NewState;
+	}
+}
+
+bool FOneSkyLocalizationServiceProvider::RetrieveProjectGroups(ELocalizationServiceOperationConcurrency::Type InConcurrency)
+{
+	// TODO: This should do the same thing as GetStateInternal but for project groups
+	// Create a data type to hold the information from a project group
+	// This logic should be moved to a function similar to GetState, but for projects
+	if (!IsEnabled())
+	{
+		return ELocalizationServiceOperationCommandResult::Failed;
+	}
+
+	// LIST PROJECT GROUPS TEST
+	//Execute(ILocalizationServiceOperation::Create<FOneSkyListProjectGroupsOperation>(), TArray<FLocalizationServiceTranslationIdentifier>(), InConcurrency);
+	//TSharedRef<FOneSkyListProjectGroupsOperation, ESPMode::ThreadSafe> ListProjectGroupsOperation = ILocalizationServiceOperation::Create<FOneSkyListProjectGroupsOperation>();
+	//ListProjectGroupsOperation->SetInStartPage(1);
+	//ListProjectGroupsOperation->SetInItemsPerPage(1);
+	//Execute(ListProjectGroupsOperation, TArray<FLocalizationServiceTranslationIdentifier>(), InConcurrency);
+
+	// SHOW PROJECT GROUP TEST
+	//TSharedRef<FOneSkyShowProjectGroupOperation, ESPMode::ThreadSafe> ShowProjectGroupOperation = ILocalizationServiceOperation::Create<FOneSkyShowProjectGroupOperation>();
+	//ShowProjectGroupOperation->SetInProjectGroupId(10222);
+	//Execute(ShowProjectGroupOperation, TArray<FLocalizationServiceTranslationIdentifier>(), InConcurrency);
+
+	// CREATE PROJECT GROUP TEST
+	//TSharedRef<FOneSkyCreateProjectGroupOperation, ESPMode::ThreadSafe> CreateProjectGroupOperation = ILocalizationServiceOperation::Create<FOneSkyCreateProjectGroupOperation>();
+	//CreateProjectGroupOperation->SetInProjectGroupName("test group");
+	//CreateProjectGroupOperation->SetInInProjectGroupBaseLocale("zh-TW");
+	//Execute(CreateProjectGroupOperation, TArray<FLocalizationServiceTranslationIdentifier>(), InConcurrency);
+
+	// LIST PROJECT GROUP LANGUAGES TEST
+	//TSharedRef<FOneSkyListProjectGroupLanguagesOperation, ESPMode::ThreadSafe> ListProjectGroupLanguagesOperation = ILocalizationServiceOperation::Create<FOneSkyListProjectGroupLanguagesOperation>();
+	//ListProjectGroupLanguagesOperation->SetInProjectGroupId(10222);
+	//Execute(ListProjectGroupLanguagesOperation, TArray<FLocalizationServiceTranslationIdentifier>(), InConcurrency);
+
+	// LIST PROJECTS IN GROUP TEST
+	//TSharedRef<FOneSkyListProjectsInGroupOperation, ESPMode::ThreadSafe> ListProjectsInGroupOperation = ILocalizationServiceOperation::Create<FOneSkyListProjectsInGroupOperation>();
+	//ListProjectsInGroupOperation->SetInProjectGroupId(10222);
+	//Execute(ListProjectsInGroupOperation, TArray<FLocalizationServiceTranslationIdentifier>(), InConcurrency);
+
+	// SHOW PROJECT test
+	//TSharedRef<FOneSkyShowProjectOperation, ESPMode::ThreadSafe> ShowProjectOperation = ILocalizationServiceOperation::Create<FOneSkyShowProjectOperation>();
+	//ShowProjectOperation->SetInProjectId(27922);
+	//Execute(ShowProjectOperation, TArray<FLocalizationServiceTranslationIdentifier>(), InConcurrency);
+
+	// CREATE PROJECT test
+	//TSharedRef<FOneSkyCreateProjectOperation, ESPMode::ThreadSafe> CreateProjectOperation = ILocalizationServiceOperation::Create<FOneSkyCreateProjectOperation>();
+	//CreateProjectOperation->SetInProjectGroupId(10222);
+	//CreateProjectOperation->SetInProjectGroupName("test create project");
+	//CreateProjectOperation->SetInInProjectDescription("test project description");
+	//CreateProjectOperation->SetInInProjectType("website");
+	//Execute(CreateProjectOperation, TArray<FLocalizationServiceTranslationIdentifier>(), InConcurrency);
+
+	// LIST PROJECT LANGUAGES TEST
+	//TSharedRef<FOneSkyListProjectLanguagesOperation, ESPMode::ThreadSafe> ListProjectLanguagesOperation = ILocalizationServiceOperation::Create<FOneSkyListProjectLanguagesOperation>();
+	//ListProjectLanguagesOperation->SetInProjectId(27922);
+	//Execute(ListProjectLanguagesOperation, TArray<FLocalizationServiceTranslationIdentifier>(), InConcurrency);
+
+	// TRANSLATION STATUS TEST
+	//TSharedRef<FOneSkyTranslationStatusOperation, ESPMode::ThreadSafe> TranslationStatusOperation = ILocalizationServiceOperation::Create<FOneSkyTranslationStatusOperation>();
+	//TranslationStatusOperation->SetInProjectId(27922);
+	//TranslationStatusOperation->SetInFileName("en.po");
+	//TranslationStatusOperation->SetInLocale("ko");
+	//Execute(TranslationStatusOperation, TArray<FLocalizationServiceTranslationIdentifier>(), InConcurrency);
+
+	// TRANSLATION EXPORT TEST
+	//TSharedRef<FOneSkyTranslationExportOperation, ESPMode::ThreadSafe> TranslationExportLanguagesOperation = ILocalizationServiceOperation::Create<FOneSkyTranslationExportOperation>();
+	//TranslationExportLanguagesOperation->SetInProjectId(27922);
+	//TranslationExportLanguagesOperation->SetInSourceFileName("en.po");
+	//TranslationExportLanguagesOperation->SetInExportFileName("D:\\out-ko.po");
+	//TranslationExportLanguagesOperation->SetInLocale("ko");
+	//Execute(TranslationExportLanguagesOperation, TArray<FLocalizationServiceTranslationIdentifier>(), InConcurrency);
+
+	// LIST UPLOADED FILES
+	//TSharedRef<FOneSkyListUploadedFilesOperation, ESPMode::ThreadSafe> ListUploadedFilesOperation = ILocalizationServiceOperation::Create<FOneSkyListUploadedFilesOperation>();
+	//ListUploadedFilesOperation->SetInProjectId(27922);
+	//ListUploadedFilesOperation->SetInItemsPerPage(1);
+	//ListUploadedFilesOperation->SetInStartPage(1);
+	//Execute(ListUploadedFilesOperation, TArray<FLocalizationServiceTranslationIdentifier>(), InConcurrency);
+
+	// UPLOAD FILE
+	//TSharedRef<FOneSkyUploadFileOperation, ESPMode::ThreadSafe> UploadFileOperation = ILocalizationServiceOperation::Create<FOneSkyUploadFileOperation>();
+	//UploadFileOperation->SetInProjectId(27922);
+	//UploadFileOperation->SetInLocalFilePathAndName("D:\\zh_CN.po");
+	//UploadFileOperation->SetInOneSkyTargetFileName("en.po");
+	//UploadFileOperation->SetInLocale("zh_CN");
+	//UploadFileOperation->SetInIsKeepingAllStrings(false);
+	//Execute(UploadFileOperation, TArray<FLocalizationServiceTranslationIdentifier>(), InConcurrency);
+
+	// LIST UPLOADED FILES
+	//TSharedRef<FOneSkyListPhraseCollectionsOperation, ESPMode::ThreadSafe> ListPhraseCollectionsOperation = ILocalizationServiceOperation::Create<FOneSkyListPhraseCollectionsOperation>();
+	//ListPhraseCollectionsOperation->SetInProjectId(27922);
+	//ListPhraseCollectionsOperation->SetInItemsPerPage(1);
+	//ListPhraseCollectionsOperation->SetInPage(1);
+	//Execute(ListPhraseCollectionsOperation, TArray<FLocalizationServiceTranslationIdentifier>(), InConcurrency);
+
+	return ELocalizationServiceOperationCommandResult::Succeeded;
+}
+
+FText FOneSkyLocalizationServiceProvider::GetStatusText() const
+{
+	FOneSkyLocalizationServiceModule& OneSkyLocalizationService = FModuleManager::LoadModuleChecked<FOneSkyLocalizationServiceModule>( "OneSkyLocalizationService" );
+	const FOneSkyLocalizationServiceSettings& Settings = OneSkyLocalizationService.AccessSettings();
+
+	FFormatNamedArguments Args;
+	Args.Add( TEXT("IsEnabled"), IsEnabled() ? LOCTEXT("Yes", "Yes") : LOCTEXT("No", "No") );
+	Args.Add( TEXT("IsConnected"), (IsEnabled() && IsAvailable()) ? LOCTEXT("Yes", "Yes") : LOCTEXT("No", "No") );
+	Args.Add( TEXT("ConnectionName"), FText::FromString( Settings.GetConnectionName() ) );
+
+	return FText::Format( LOCTEXT("OneSkyStatusText", "Enabled: {IsEnabled}\nConnected: {IsConnected}\nConnectionName: {ConnectionName}\n"), Args );
+}
+
+bool FOneSkyLocalizationServiceProvider::IsEnabled() const
+{
+	return true;
+}
+
+bool FOneSkyLocalizationServiceProvider::IsAvailable() const
+{
+	return bServerAvailable;
+}
+
+bool FOneSkyLocalizationServiceProvider::EstablishPersistentConnection()
+{
+	FOneSkyLocalizationServiceModule& OneSkyLocalizationService = FModuleManager::LoadModuleChecked<FOneSkyLocalizationServiceModule>( "OneSkyLocalizationService" );
+	FOneSkyConnectionInfo ConnectionInfo = OneSkyLocalizationService.AccessSettings().GetConnectionInfo();
+
+	bool bIsValidConnection = false;
+	if ( !PersistentConnection )
+	{
+		PersistentConnection = new FOneSkyConnection(ConnectionInfo);
+	}
+
+	bIsValidConnection = PersistentConnection->IsValidConnection();
+	if ( !bIsValidConnection )
+	{
+		delete PersistentConnection;
+		PersistentConnection = new FOneSkyConnection(ConnectionInfo);
+		bIsValidConnection = PersistentConnection->IsValidConnection();
+	}
+
+	bServerAvailable = bIsValidConnection;
+	return bIsValidConnection;
+}
+
+const FName& FOneSkyLocalizationServiceProvider::GetName() const
+{
+	return ProviderName;
+}
+
+const FText FOneSkyLocalizationServiceProvider::GetDisplayName() const
+{
+	return LOCTEXT("OneSkyLocalizationService", "OneSky Localization Service");
+}
+
+ELocalizationServiceOperationCommandResult::Type FOneSkyLocalizationServiceProvider::GetState(const TArray<FLocalizationServiceTranslationIdentifier>& InTranslationIds, TArray< TSharedRef<ILocalizationServiceState, ESPMode::ThreadSafe> >& OutState, ELocalizationServiceCacheUsage::Type InStateCacheUsage)
+{
+	if(!IsEnabled())
+	{
+		return ELocalizationServiceOperationCommandResult::Failed;
+	}
+
+	if (InStateCacheUsage == ELocalizationServiceCacheUsage::ForceUpdate)
+	{
+		// TODO: Something like this?
+		//for (TArray<FText>::TConstIterator It(InTexts); It; It++)
+		{
+			//TSharedRef<FOneSkyShowPhraseCollectionWorker, ESPMode::ThreadSafe> PhraseCollectionWorker = ILocalizationServiceOperation::Create<FOneSkyShowPhraseCollectionWorker>();
+
+			//PhraseCollectionWorker->InCollectionKey = FTextInspector::GetKey(*It);
+			//Execute(ILocalizationServiceOperation::Create<FLocalizationServiceUpdateStatus>(), InTranslationIds);
+		}
+	}
+
+	for (TArray<FLocalizationServiceTranslationIdentifier>::TConstIterator It(InTranslationIds); It; It++)
+	{
+		TSharedRef<FOneSkyLocalizationServiceState, ESPMode::ThreadSafe>* State = StateCache.Find(*It);
+		if(State != NULL)
+		{
+			// found cached item for this file, return that
+			OutState.Add(*State);
+		}
+		else
+		{
+			// cache an unknown state for this item & return that
+			TSharedRef<FOneSkyLocalizationServiceState, ESPMode::ThreadSafe> NewState = MakeShareable(new FOneSkyLocalizationServiceState(*It));
+			StateCache.Add(*It, NewState);
+			OutState.Add(NewState);
+		}
+	}
+
+	return ELocalizationServiceOperationCommandResult::Succeeded;
+}
+
+ELocalizationServiceOperationCommandResult::Type FOneSkyLocalizationServiceProvider::Execute(const TSharedRef<ILocalizationServiceOperation, ESPMode::ThreadSafe>& InOperation, const TArray<FLocalizationServiceTranslationIdentifier>& InTranslationIds, ELocalizationServiceOperationConcurrency::Type InConcurrency, const FLocalizationServiceOperationComplete& InOperationCompleteDelegate)
+{
+	if(!IsEnabled())
+	{
+		return ELocalizationServiceOperationCommandResult::Failed;
+	}
+
+	// Query to see if the we allow this operation
+	TSharedPtr<IOneSkyLocalizationServiceWorker, ESPMode::ThreadSafe> Worker = CreateWorker(InOperation->GetName());
+	if(!Worker.IsValid())
+	{
+		// this operation is unsupported by this source control provider
+		FFormatNamedArguments Arguments;
+		Arguments.Add( TEXT("OperationName"), FText::FromName(InOperation->GetName()) );
+		Arguments.Add( TEXT("ProviderName"), FText::FromName(GetName()) );
+		FMessageLog("LocalizationService").Error(FText::Format(LOCTEXT("UnsupportedOperation", "Operation '{OperationName}' not supported by source control provider '{ProviderName}'"), Arguments));
+		return ELocalizationServiceOperationCommandResult::Failed;
+	}
+
+	// fire off operation
+	if(InConcurrency == ELocalizationServiceOperationConcurrency::Synchronous)
+	{
+		FOneSkyLocalizationServiceCommand* Command = new FOneSkyLocalizationServiceCommand(InOperation, Worker.ToSharedRef());
+		Command->bAutoDelete = false;
+		Command->OperationCompleteDelegate = InOperationCompleteDelegate;
+		return ExecuteSynchronousCommand(*Command, InOperation->GetInProgressString(), true);
+	}
+	else
+	{
+		FOneSkyLocalizationServiceCommand* Command = new FOneSkyLocalizationServiceCommand(InOperation, Worker.ToSharedRef());
+		Command->bAutoDelete = true;
+		Command->OperationCompleteDelegate = InOperationCompleteDelegate;
+		return IssueCommand(*Command, false);
+	}
+}
+
+bool FOneSkyLocalizationServiceProvider::CanCancelOperation( const TSharedRef<ILocalizationServiceOperation, ESPMode::ThreadSafe>& InOperation ) const
+{
+	for (int32 CommandIndex = 0; CommandIndex < CommandQueue.Num(); ++CommandIndex)
+	{
+		FOneSkyLocalizationServiceCommand& Command = *CommandQueue[CommandIndex];
+		if (Command.Operation == InOperation)
+		{
+			check(Command.bAutoDelete);
+			return true;
+		}
+	}
+
+	// operation was not in progress!
+	return false;
+}
+
+void FOneSkyLocalizationServiceProvider::CancelOperation( const TSharedRef<ILocalizationServiceOperation, ESPMode::ThreadSafe>& InOperation )
+{
+	for (int32 CommandIndex = 0; CommandIndex < CommandQueue.Num(); ++CommandIndex)
+	{
+		FOneSkyLocalizationServiceCommand& Command = *CommandQueue[CommandIndex];
+		if (Command.Operation == InOperation)
+		{
+			check(Command.bAutoDelete);
+			Command.Cancel();
+			return;
+		}
+	}
+}
+
+void FOneSkyLocalizationServiceProvider::OutputCommandMessages(const FOneSkyLocalizationServiceCommand& InCommand) const
+{
+	FMessageLog LocalizationServiceLog("LocalizationService");
+
+	for (int32 ErrorIndex = 0; ErrorIndex < InCommand.ErrorMessages.Num(); ++ErrorIndex)
+	{
+		LocalizationServiceLog.Error(InCommand.ErrorMessages[ErrorIndex]);
+	}
+
+	for (int32 InfoIndex = 0; InfoIndex < InCommand.InfoMessages.Num(); ++InfoIndex)
+	{
+		LocalizationServiceLog.Info(InCommand.InfoMessages[InfoIndex]);
+	}
+}
+
+void FOneSkyLocalizationServiceProvider::Tick()
+{	
+	bool bStatesUpdated = false;
+	for (int32 CommandIndex = 0; CommandIndex < CommandQueue.Num(); ++CommandIndex)
+	{
+		FOneSkyLocalizationServiceCommand& Command = *CommandQueue[CommandIndex];
+		if (Command.bExecuteProcessed)
+		{
+			// Remove command from the queue
+			CommandQueue.RemoveAt(CommandIndex);
+
+			// update connection state
+			bServerAvailable = !Command.bConnectionDropped || Command.bCancelled;
+
+			// let command update the states of any files
+			bStatesUpdated |= Command.Worker->UpdateStates();
+
+			// dump any messages to output log
+			OutputCommandMessages(Command);
+
+			// run the completion delegate if we have one bound
+			ELocalizationServiceOperationCommandResult::Type Result = ELocalizationServiceOperationCommandResult::Failed;
+			if(Command.bCommandSuccessful)
+			{
+				Result = ELocalizationServiceOperationCommandResult::Succeeded;
+			}
+			else if(Command.bCancelled)
+			{
+				Result = ELocalizationServiceOperationCommandResult::Cancelled;
+			}
+			Command.OperationCompleteDelegate.ExecuteIfBound(Command.Operation, Result);
+
+			//commands that are left in the array during a tick need to be deleted
+			if(Command.bAutoDelete)
+			{
+				// Only delete commands that are not running 'synchronously'
+				delete &Command;
+			}
+
+			// only do one command per tick loop, as we dont want concurrent modification 
+			// of the command queue (which can happen in the completion delegate)
+			break;
+		}
+	}
+
+	//if(bStatesUpdated)
+	//{
+	//	OnLocalizationServiceStateChanged.Broadcast();
+	//}
+}
+
+static void PublicKeyChanged(const FText& NewText, ETextCommit::Type CommitType)
+{
+	FOneSkyLocalizationServiceModule::Get().AccessSettings().SetApiKey(NewText.ToString());
+	FOneSkyLocalizationServiceModule::Get().AccessSettings().SaveSettings();
+}
+
+static void SecretKeyChanged(const FText& NewText, ETextCommit::Type CommitType)
+{
+	FOneSkyLocalizationServiceModule::Get().AccessSettings().SetApiSecret(NewText.ToString());
+	FOneSkyLocalizationServiceModule::Get().AccessSettings().SaveSettings();
+}
+
+static void SaveSecretKeyChanged(ECheckBoxState CheckBoxState)
+{
+	FOneSkyLocalizationServiceModule::Get().AccessSettings().SetSaveSecretKey(CheckBoxState == ECheckBoxState::Checked);
+	FOneSkyLocalizationServiceModule::Get().AccessSettings().SaveSettings();
+}
+
+#if LOCALIZATION_SERVICES_WITH_SLATE
+void FOneSkyLocalizationServiceProvider::CustomizeSettingsDetails(IDetailCategoryBuilder& DetailCategoryBuilder) const
+{
+	FOneSkyConnectionInfo ConnectionInfo = FOneSkyLocalizationServiceModule::Get().AccessSettings().GetConnectionInfo();
+	FText PublicKeyText = LOCTEXT("OneSkyPublicKeyLabel", "OneSky API Public Key");
+	FDetailWidgetRow& PublicKeyRow = DetailCategoryBuilder.AddCustomRow(PublicKeyText);
+	PublicKeyRow.NameContent()
+		[
+			SNew(STextBlock)
+			.Text(PublicKeyText)
+		];
+	PublicKeyRow.ValueContent()
+		[
+			SNew(SEditableTextBox)
+			.OnTextCommitted(FOnTextCommitted::CreateStatic(&PublicKeyChanged))
+			.Text(FText::FromString(ConnectionInfo.ApiKey))
+		];
+
+	FText SecretKeyText = LOCTEXT("OneSkySecretKeyLabel", "OneSky API Secret Key");
+	FDetailWidgetRow& SecretKeyRow = DetailCategoryBuilder.AddCustomRow(SecretKeyText);
+	SecretKeyRow.NameContent()
+		[
+			SNew(STextBlock)
+			.Text(SecretKeyText)
+		];
+	SecretKeyRow.ValueContent()
+		[
+			SNew(SEditableTextBox)
+			.OnTextCommitted(FOnTextCommitted::CreateStatic(&SecretKeyChanged))
+			.Text(FText::FromString(ConnectionInfo.ApiSecret))
+		];
+
+	FText SaveSecretKeyText = LOCTEXT("OneSkySaveSecret", "Remember Secret Key (WARNING: saved unencrypted)");
+	FDetailWidgetRow& SaveSecretKeyRow = DetailCategoryBuilder.AddCustomRow(SaveSecretKeyText);
+	SaveSecretKeyRow.NameContent()
+		[
+			SNew(STextBlock)
+			.Text(SaveSecretKeyText)
+		];
+	SaveSecretKeyRow.ValueContent()
+		[
+			SNew(SCheckBox)
+			.IsChecked(FOneSkyLocalizationServiceModule::Get().AccessSettings().GetSaveSecretKey() ?  ECheckBoxState::Checked : ECheckBoxState::Unchecked)
+			.OnCheckStateChanged(FOnCheckStateChanged::CreateStatic(&SaveSecretKeyChanged))
+		];
+}
+
+static void ProjectGroupChanged(const FText& NewText, ETextCommit::Type CommitType, FGuid TargetGuid)
+{
+	FOneSkyLocalizationTargetSetting* Settings = FOneSkyLocalizationServiceModule::Get().AccessSettings().GetSettingsForTarget(TargetGuid);
+	FOneSkyLocalizationServiceModule::Get().AccessSettings().SetSettingsForTarget(TargetGuid, NewText.ToString(), Settings->OneSkyProjectName, Settings->OneSkyFileName);
+}
+
+static void ProjectChanged(const FText& NewText, ETextCommit::Type CommitType, FGuid TargetGuid)
+{
+	FOneSkyLocalizationTargetSetting* Settings = FOneSkyLocalizationServiceModule::Get().AccessSettings().GetSettingsForTarget(TargetGuid);
+	FOneSkyLocalizationServiceModule::Get().AccessSettings().SetSettingsForTarget(TargetGuid, Settings->OneSkyProjectGroupName, NewText.ToString(), Settings->OneSkyFileName);
+}
+
+static void FileNameChanged(const FText& NewText, ETextCommit::Type CommitType, FGuid TargetGuid)
+{
+	FOneSkyLocalizationTargetSetting* Settings = FOneSkyLocalizationServiceModule::Get().AccessSettings().GetSettingsForTarget(TargetGuid);
+	FOneSkyLocalizationServiceModule::Get().AccessSettings().SetSettingsForTarget(TargetGuid, Settings->OneSkyProjectGroupName, Settings->OneSkyProjectName, NewText.ToString());
+}
+
+void FOneSkyLocalizationServiceProvider::CustomizeTargetDetails(IDetailCategoryBuilder& DetailCategoryBuilder, const FGuid& TargetGuid) const
+{
+	//Target.Settings.Guid.ToString();
+	FOneSkyLocalizationTargetSetting* Settings = FOneSkyLocalizationServiceModule::Get().AccessSettings().GetSettingsForTarget(TargetGuid);
+
+	FText ProjectGroupText = LOCTEXT("OneSkyProjectGroupLabel", "OneSky Project Group");
+	FDetailWidgetRow& ProjectGroupRow = DetailCategoryBuilder.AddCustomRow(ProjectGroupText);
+	ProjectGroupRow.NameContent()
+		[
+			SNew(STextBlock)
+			.Text(ProjectGroupText)
+		];
+	ProjectGroupRow.ValueContent()
+		[
+			SNew(SEditableTextBox)
+			.OnTextCommitted(FOnTextCommitted::CreateStatic(&ProjectGroupChanged, TargetGuid))
+			.Text(FText::FromString(Settings->OneSkyProjectGroupName))
+		];
+
+	FText ProjectText = LOCTEXT("OneSkyProjectLabel", "OneSky Project");
+	FDetailWidgetRow& ProjectRow = DetailCategoryBuilder.AddCustomRow(ProjectText);
+	ProjectRow.NameContent()
+		[
+			SNew(STextBlock)
+			.Text(ProjectText)
+		];
+	ProjectRow.ValueContent()
+		[
+			SNew(SEditableTextBox)
+			.OnTextCommitted(FOnTextCommitted::CreateStatic(&ProjectChanged, TargetGuid))
+			.Text(FText::FromString(Settings->OneSkyProjectName))
+		];
+
+	FText FileText = LOCTEXT("OneSkyFileNameLabel", "OneSky File Name");
+	FDetailWidgetRow& FileNameRow = DetailCategoryBuilder.AddCustomRow(FileText);
+	FileNameRow.NameContent()
+		[
+			SNew(STextBlock)
+			.Text(FileText)
+		];
+	FileNameRow.ValueContent()
+		[
+			SNew(SEditableTextBox)
+			.OnTextCommitted(FOnTextCommitted::CreateStatic(&FileNameChanged, TargetGuid))
+			.Text(FText::FromString(Settings->OneSkyFileName))
+		];
+}
+#endif //LOCALIZATION_SERVICES_WITH_SLATE
+
+TSharedPtr<IOneSkyLocalizationServiceWorker, ESPMode::ThreadSafe> FOneSkyLocalizationServiceProvider::CreateWorker(const FName& InOperationName) const
+{
+	const FGetOneSkyLocalizationServiceWorker* Operation = WorkersMap.Find(InOperationName);
+	if(Operation != NULL)
+	{
+		return Operation->Execute();
+	}
+		
+	return NULL;
+}
+
+void FOneSkyLocalizationServiceProvider::RegisterWorker( const FName& InName, const FGetOneSkyLocalizationServiceWorker& InDelegate )
+{
+	WorkersMap.Add( InName, InDelegate );
+}
+
+ELocalizationServiceOperationCommandResult::Type FOneSkyLocalizationServiceProvider::ExecuteSynchronousCommand(FOneSkyLocalizationServiceCommand& InCommand, const FText& Task, bool bSuppressResponseMsg)
+{
+	ELocalizationServiceOperationCommandResult::Type Result = ELocalizationServiceOperationCommandResult::Failed;
+
+	struct Local
+	{
+		static void CancelCommand(FOneSkyLocalizationServiceCommand* InControlCommand)
+		{
+			InControlCommand->Cancel();
+		}
+	};
+
+	// Display the progress dialog
+	//FScopedLocalizationServiceProgress Progress(Task, FSimpleDelegate::CreateStatic(&Local::CancelCommand, &InCommand));
+
+	// Perform the command asynchronously
+	IssueCommand( InCommand, false );
+
+	// Wait until the queue is empty. Only at this point is our command guaranteed to be removed from the queue
+	while(CommandQueue.Num() > 0)
+	{
+		// Tick the command queue and update progress.
+		Tick();
+
+		//Progress.Tick();
+
+		// Sleep for a bit so we don't busy-wait so much.
+		FPlatformProcess::Sleep(0.01f);
+	}
+
+	if(InCommand.bCommandSuccessful)
+	{
+		Result = ELocalizationServiceOperationCommandResult::Succeeded;
+	}
+	else if(InCommand.bCancelled)
+	{
+		Result = ELocalizationServiceOperationCommandResult::Cancelled;
+	}
+
+	// If the command failed, inform the user that they need to try again
+	if (!InCommand.bCancelled && Result != ELocalizationServiceOperationCommandResult::Succeeded && !bSuppressResponseMsg)
+	{
+		FMessageDialog::Open( EAppMsgType::Ok, LOCTEXT("OneSky_ServerUnresponsive", "OneSky server is unresponsive. Please check your connection and try again.") );
+	}
+
+	// Delete the command now
+	check(!InCommand.bAutoDelete);
+	delete &InCommand;
+
+	return Result;
+}
+
+ELocalizationServiceOperationCommandResult::Type FOneSkyLocalizationServiceProvider::IssueCommand(FOneSkyLocalizationServiceCommand& InCommand, const bool bSynchronous)
+{
+	if ( !bSynchronous && GThreadPool != NULL )
+	{
+		// Queue this to our worker thread(s) for resolving
+		GThreadPool->AddQueuedWork(&InCommand);
+		CommandQueue.Add(&InCommand);
+		return ELocalizationServiceOperationCommandResult::Succeeded;
+	}
+	else
+	{
+		InCommand.bCommandSuccessful = InCommand.DoWork();
+
+		InCommand.Worker->UpdateStates();
+
+		OutputCommandMessages(InCommand);
+
+		// Callback now if present. When asynchronous, this callback gets called from Tick().
+		ELocalizationServiceOperationCommandResult::Type Result = InCommand.bCommandSuccessful ? ELocalizationServiceOperationCommandResult::Succeeded : ELocalizationServiceOperationCommandResult::Failed;
+		InCommand.OperationCompleteDelegate.ExecuteIfBound(InCommand.Operation, Result);
+
+		return Result;
+	}
+}
+
+
+#undef LOCTEXT_NAMESPACE
diff --git a/Engine/Plugins/Developer/OneSkyLocalizationService/Source/OneSkyLocalizationService/Private/OneSkyLocalizationServiceProvider.h b/Engine/Plugins/Developer/OneSkyLocalizationService/Source/OneSkyLocalizationService/Private/OneSkyLocalizationServiceProvider.h
new file mode 100644
--- /dev/null
+++ b/Engine/Plugins/Developer/OneSkyLocalizationService/Source/OneSkyLocalizationService/Private/OneSkyLocalizationServiceProvider.h
@@ -0,0 +1,146 @@
+// Copyright 1998-2015 Epic Games, Inc. All Rights Reserved.
+
+#pragma once
+
+#include "Developer/LocalizationService/Public/ILocalizationServiceModule.h"
+#include "IOneSkyLocalizationServiceWorker.h"
+#include "OneSkyLocalizationServiceState.h"
+
+DECLARE_DELEGATE_RetVal(FOneSkyLocalizationServiceWorkerRef, FGetOneSkyLocalizationServiceWorker)
+
+// Define how the TMap from FLocalizationServiceTranslationIdentifier works
+template <typename ValueType>
+struct FLocalizationServiceTranslationIdentifierKeyFuncs : BaseKeyFuncs < TPair< FLocalizationServiceTranslationIdentifier, ValueType>, FLocalizationServiceTranslationIdentifier, false >
+{
+	typedef BaseKeyFuncs <
+		TPair<FLocalizationServiceTranslationIdentifier, ValueType>,
+		FLocalizationServiceTranslationIdentifier
+	> Super;
+	typedef typename Super::ElementInitType ElementInitType;
+	typedef typename Super::KeyInitType     KeyInitType;
+	static FORCEINLINE KeyInitType GetSetKey(ElementInitType Element)
+	{
+		return Element.Key;
+	}
+	static FORCEINLINE bool Matches(const FLocalizationServiceTranslationIdentifier& A, const FLocalizationServiceTranslationIdentifier& B)
+	{
+		return A.Culture.IsValid() && B.Culture.IsValid() &&
+			A.LocalizationTargetGuid == B.LocalizationTargetGuid &&
+			A.Culture->GetName().Equals(B.Culture->GetName()) &&
+			A.Namespace.Equals(B.Namespace, ESearchCase::CaseSensitive) &&
+			A.Source.Equals(B.Source, ESearchCase::CaseSensitive);
+	}
+	static FORCEINLINE uint32 GetKeyHash(const FLocalizationServiceTranslationIdentifier& Key)
+	{
+		if (!Key.Culture.IsValid())
+		{
+			return 0;
+		}
+
+		uint32 GuidHash = GetTypeHash(Key.LocalizationTargetGuid);
+		uint32 CultureHash = FCrc::StrCrc32<TCHAR>(*Key.Culture->GetName());
+		uint32 NamespaceHash = FCrc::StrCrc32<TCHAR>(*Key.Namespace);
+		uint32 SourceHash = FCrc::StrCrc32<TCHAR>(*Key.Source);
+
+		return HashCombine(HashCombine(HashCombine(GuidHash, CultureHash), NamespaceHash), SourceHash);
+	}
+};
+
+class FOneSkyLocalizationServiceProvider : public ILocalizationServiceProvider
+{
+public:
+	/** Constructor */
+	FOneSkyLocalizationServiceProvider()
+		: bServerAvailable(false)
+		, PersistentConnection(NULL)
+	{
+	}
+
+	/* ILocalizationServiceProvider implementation */
+	virtual void Init(bool bForceConnection = true) override;
+	virtual void Close() override;
+	virtual FText GetStatusText() const override;
+	virtual bool IsEnabled() const override;
+	virtual bool IsAvailable() const override;
+	virtual const FName& GetName(void) const override;
+	virtual const FText GetDisplayName() const override;
+	virtual ELocalizationServiceOperationCommandResult::Type GetState(const TArray<FLocalizationServiceTranslationIdentifier>& InTranslationIds, TArray< TSharedRef<ILocalizationServiceState, ESPMode::ThreadSafe> >& OutState, ELocalizationServiceCacheUsage::Type InStateCacheUsage) override;
+	virtual ELocalizationServiceOperationCommandResult::Type Execute(const TSharedRef<ILocalizationServiceOperation, ESPMode::ThreadSafe>& InOperation, const TArray<FLocalizationServiceTranslationIdentifier>& InTranslationIds, ELocalizationServiceOperationConcurrency::Type InConcurrency = ELocalizationServiceOperationConcurrency::Synchronous, const FLocalizationServiceOperationComplete& InOperationCompleteDelegate = FLocalizationServiceOperationComplete()) override;
+	virtual bool CanCancelOperation( const TSharedRef<ILocalizationServiceOperation, ESPMode::ThreadSafe>& InOperation ) const override;
+	virtual void CancelOperation( const TSharedRef<ILocalizationServiceOperation, ESPMode::ThreadSafe>& InOperation ) override;
+	virtual void Tick() override;
+#if LOCALIZATION_SERVICES_WITH_SLATE
+	virtual void CustomizeSettingsDetails(IDetailCategoryBuilder& DetailCategoryBuilder) const override;
+	virtual void CustomizeTargetDetails(IDetailCategoryBuilder& DetailCategoryBuilder, const FGuid& TargetGuid) const override;
+#endif // LOCALIZATION_SERVICES_WITH_SLATE
+
+	/**
+	 * Register a worker with the provider.
+	 * This is used internally so the provider can maintain a map of all available operations.
+	 */
+	void RegisterWorker( const FName& InName, const FGetOneSkyLocalizationServiceWorker& InDelegate );
+
+	/** Helper function used to update state cache */
+	TSharedRef<FOneSkyLocalizationServiceState, ESPMode::ThreadSafe> GetStateInternal(const FLocalizationServiceTranslationIdentifier& InTranslationId);
+
+	/** Helper function used to update project groups cache */
+	bool RetrieveProjectGroups(ELocalizationServiceOperationConcurrency::Type InConcurrency);
+
+	/**
+	 * Connects to the Localization service server if the persistent connection is not already established.
+	 *
+	 * @return true if the connection is established or became established and false if the connection failed.
+	 */
+	bool EstablishPersistentConnection();
+
+	/** Get the persistent connection, if any */
+	class FOneSkyConnection* GetPersistentConnection()
+	{
+		return PersistentConnection;
+	}
+
+private:
+
+	/** Helper function used to create a worker for a particular operation */
+	TSharedPtr<class IOneSkyLocalizationServiceWorker, ESPMode::ThreadSafe> CreateWorker(const FName& InOperationName) const;
+
+	/** 
+	 * Logs any messages that a command needs to output.
+	 */
+	void OutputCommandMessages(const class FOneSkyLocalizationServiceCommand& InCommand) const;
+
+	/**
+	 * Helper function for running command 'synchronously'.
+	 * This really doesn't execute synchronously; rather it adds the command to the queue & does not return until
+	 * the command is completed.
+	 */
+	ELocalizationServiceOperationCommandResult::Type ExecuteSynchronousCommand(class FOneSkyLocalizationServiceCommand& InCommand, const FText& Task, bool bSuppressResponseMsg);
+
+	/**
+	 * Run a command synchronously or asynchronously.
+	 */
+	ELocalizationServiceOperationCommandResult::Type IssueCommand(class FOneSkyLocalizationServiceCommand& InCommand, const bool bSynchronous);
+
+private:
+
+	/** Indicates if Localization service integration is available or not. */
+	bool bServerAvailable;
+
+	/** A pointer to the persistent P4 connection for synchronous operations */
+	class FOneSkyConnection* PersistentConnection;
+
+	/** Project Group Cache (group name to group id) */
+	TMap<FString, int32> ProjectGroupCache;
+
+	/** Project Cache (project name to project id) per group id */
+	TMap<int32, TMap<FString, int32>> ProjectCachePerGroupId;
+	
+	/** State cache */
+	TMap<FLocalizationServiceTranslationIdentifier, TSharedRef<class FOneSkyLocalizationServiceState, ESPMode::ThreadSafe>, FDefaultSetAllocator, FLocalizationServiceTranslationIdentifierKeyFuncs<TSharedRef<class FOneSkyLocalizationServiceState, ESPMode::ThreadSafe>> >  StateCache;
+
+	/** The currently registered Localization service operations */
+	TMap<FName, FGetOneSkyLocalizationServiceWorker> WorkersMap;
+
+	/** Queue for commands given by the main thread */
+	TArray < FOneSkyLocalizationServiceCommand* > CommandQueue;
+};
diff --git a/Engine/Plugins/Developer/OneSkyLocalizationService/Source/OneSkyLocalizationService/Private/OneSkyLocalizationServiceResponseTypes.h b/Engine/Plugins/Developer/OneSkyLocalizationService/Source/OneSkyLocalizationService/Private/OneSkyLocalizationServiceResponseTypes.h
new file mode 100644
--- /dev/null
+++ b/Engine/Plugins/Developer/OneSkyLocalizationService/Source/OneSkyLocalizationService/Private/OneSkyLocalizationServiceResponseTypes.h
@@ -0,0 +1,1207 @@
+// Copyright 1998-2015 Epic Games, Inc. All Rights Reserved.
+
+#pragma once
+
+#include "OneSkyLocalizationServiceResponseTypes.generated.h"
+
+
+// LIST PROJECT GROUPS
+
+USTRUCT()
+struct FOneSkyListProjectGroupsResponseMeta
+{
+	GENERATED_USTRUCT_BODY()
+
+	UPROPERTY()
+		int32 status;
+
+	UPROPERTY()
+		int32 record_count;
+
+	UPROPERTY()
+		int32 page_count;
+
+	UPROPERTY()
+		FString next_page;
+
+	UPROPERTY()
+		FString prev_page;
+
+	UPROPERTY()
+		FString first_page;
+
+	UPROPERTY()
+		FString last_page;
+
+	/** Default constructor. */
+	FOneSkyListProjectGroupsResponseMeta()
+		: status(-1)
+		, record_count(-1)
+		, page_count(-1)
+		, next_page(TEXT(""))
+		, prev_page(TEXT(""))
+		, first_page(TEXT(""))
+		, last_page(TEXT(""))
+	{ }
+
+	/** Creates an uninitialized instance. */
+	FOneSkyListProjectGroupsResponseMeta(ENoInit) { }
+};
+
+
+USTRUCT()
+struct FOneSkyListProjectGroupsResponseDataItem
+{
+	GENERATED_USTRUCT_BODY()
+
+	UPROPERTY()
+		FString id;
+
+	UPROPERTY()
+		FString name;
+
+	/** Default constructor. */
+	FOneSkyListProjectGroupsResponseDataItem()
+		: id("-1")
+		, name(TEXT(""))
+	{ }
+
+	/** Creates an uninitialized instance. */
+	FOneSkyListProjectGroupsResponseDataItem(ENoInit) { }
+};
+
+/**
+* Response from a List Project Groups query on OneSky
+* //[REDACTED_PROFILE_URL]/api-documentation-platform/blob/master/resources/project_group.md
+*/
+USTRUCT()
+struct FOneSkyListProjectGroupsResponse
+{
+	GENERATED_USTRUCT_BODY()
+
+	UPROPERTY()
+		FOneSkyListProjectGroupsResponseMeta meta;
+
+	UPROPERTY()
+		TArray<FOneSkyListProjectGroupsResponseDataItem> data;
+
+	/** Default constructor. */
+	FOneSkyListProjectGroupsResponse()
+		: meta()
+		, data()
+	{ }
+
+	/** Creates an uninitialized instance. */
+	FOneSkyListProjectGroupsResponse(ENoInit) { }
+};
+
+// end LIST PROJECT GROUPS
+
+
+// SHOW PROJECT GROUP
+
+USTRUCT()
+struct FOneSkyShowProjectGroupResponseMeta
+{
+	GENERATED_USTRUCT_BODY()
+
+		UPROPERTY()
+		int32 status;
+
+	/** Default constructor. */
+	FOneSkyShowProjectGroupResponseMeta()
+		: status(-1)
+	{ }
+
+	/** Creates an uninitialized instance. */
+	FOneSkyShowProjectGroupResponseMeta(ENoInit) { }
+};
+
+USTRUCT()
+struct FOneSkyShowProjectGroupResponseData
+{
+	GENERATED_USTRUCT_BODY()
+
+	UPROPERTY()
+		FString id;
+
+	UPROPERTY()
+		FString name;
+
+	UPROPERTY()
+		int32 description;
+
+	UPROPERTY()
+		FString project_count;
+
+	/** Default constructor. */
+	FOneSkyShowProjectGroupResponseData()
+		: id("-1")
+		, name(TEXT(""))
+		, description(-1)
+		, project_count("-1")
+	{ }
+
+	/** Creates an uninitialized instance. */
+	FOneSkyShowProjectGroupResponseData(ENoInit) { }
+};
+
+/**
+* Response from a Show Project Group query on OneSky
+* //[REDACTED_PROFILE_URL]/api-documentation-platform/blob/master/resources/project_group.md
+*/
+USTRUCT()
+struct FOneSkyShowProjectGroupResponse
+{
+	GENERATED_USTRUCT_BODY()
+
+		UPROPERTY()
+		FOneSkyShowProjectGroupResponseMeta meta;
+
+	UPROPERTY()
+		FOneSkyShowProjectGroupResponseData data;
+
+	/** Default constructor. */
+	FOneSkyShowProjectGroupResponse()
+		: meta()
+		, data()
+	{ }
+
+	/** Creates an uninitialized instance. */
+	FOneSkyShowProjectGroupResponse(ENoInit) { }
+};
+
+// end SHOW PROJECT GROUP
+
+
+// CREATE PROJECT GROUP
+
+USTRUCT()
+struct FOneSkyCreateProjectGroupResponseMeta
+{
+	GENERATED_USTRUCT_BODY()
+
+		UPROPERTY()
+		int32 status;
+
+	/** Default constructor. */
+	FOneSkyCreateProjectGroupResponseMeta()
+		: status(-1)
+	{ }
+
+	/** Creates an uninitialized instance. */
+	FOneSkyCreateProjectGroupResponseMeta(ENoInit) { }
+};
+
+USTRUCT()
+struct FOneSkyCreateProjectGroupResponseBaseLanguage
+{
+	GENERATED_USTRUCT_BODY()
+
+		UPROPERTY()
+		FString code;
+
+	UPROPERTY()
+		FString english_name;
+
+	UPROPERTY()
+		FString local_name;
+
+	UPROPERTY()
+		FString locale;
+
+	UPROPERTY()
+		FString region;
+
+	/** Default constructor. */
+	FOneSkyCreateProjectGroupResponseBaseLanguage()
+		: code("")
+		, english_name("")
+		, local_name("")
+		, locale("")
+		, region("")
+	{ }
+
+	/** Creates an uninitialized instance. */
+	FOneSkyCreateProjectGroupResponseBaseLanguage(ENoInit) { }
+};
+
+USTRUCT()
+struct FOneSkyCreateProjectGroupResponseData
+{
+	GENERATED_USTRUCT_BODY()
+
+		UPROPERTY()
+		int32 id;
+
+	UPROPERTY()
+		FString name;
+
+	UPROPERTY()
+		FOneSkyCreateProjectGroupResponseBaseLanguage base_language;
+
+	/** Default constructor. */
+	FOneSkyCreateProjectGroupResponseData()
+		: id(-1)
+		, name(TEXT(""))
+		, base_language()
+	{ }
+
+	/** Creates an uninitialized instance. */
+	FOneSkyCreateProjectGroupResponseData(ENoInit) { }
+};
+
+/**
+* Response from a Create Project Group query on OneSky
+* //[REDACTED_PROFILE_URL]/api-documentation-platform/blob/master/resources/project_group.md
+*/
+USTRUCT()
+struct FOneSkyCreateProjectGroupResponse
+{
+	GENERATED_USTRUCT_BODY()
+
+		UPROPERTY()
+		FOneSkyCreateProjectGroupResponseMeta meta;
+
+	UPROPERTY()
+		FOneSkyCreateProjectGroupResponseData data;
+
+	/** Default constructor. */
+	FOneSkyCreateProjectGroupResponse()
+		: meta()
+		, data()
+	{ }
+
+	/** Creates an uninitialized instance. */
+	FOneSkyCreateProjectGroupResponse(ENoInit) { }
+};
+
+// end CREATE PROJECT GROUP
+
+
+// LIST PROJECT GROUP LANGUAGES
+
+USTRUCT()
+struct FOneSkyListProjectGroupLanguagesResponseMeta
+{
+	GENERATED_USTRUCT_BODY()
+
+		UPROPERTY()
+		int32 status;
+
+	UPROPERTY()
+		int32 record_count;
+
+	/** Default constructor. */
+	FOneSkyListProjectGroupLanguagesResponseMeta()
+		: status(-1)
+		, record_count(-1)
+	{ }
+
+	/** Creates an uninitialized instance. */
+	FOneSkyListProjectGroupLanguagesResponseMeta(ENoInit) { }
+};
+
+USTRUCT()
+struct FOneSkyListProjectGroupLanguagesResponseDataItem
+{
+	GENERATED_USTRUCT_BODY()
+
+		UPROPERTY()
+		FString code;
+
+	UPROPERTY()
+		FString english_name;
+
+	UPROPERTY()
+		FString local_name;
+
+	UPROPERTY()
+		FString locale;
+
+	UPROPERTY()
+		FString region;
+
+	UPROPERTY()
+		bool is_base_language;
+
+	/** Default constructor. */
+	FOneSkyListProjectGroupLanguagesResponseDataItem()
+		: code("")
+		, english_name("")
+		, local_name("")
+		, locale("")
+		, region("")
+		, is_base_language(false)
+	{ }
+
+	/** Creates an uninitialized instance. */
+	FOneSkyListProjectGroupLanguagesResponseDataItem(ENoInit) { }
+};
+
+/**
+* Response from a List Project Group Languages query on OneSky
+* //[REDACTED_PROFILE_URL]/api-documentation-platform/blob/master/resources/project_group.md
+*/
+USTRUCT()
+struct FOneSkyListProjectGroupLanguagesResponse
+{
+	GENERATED_USTRUCT_BODY()
+
+		UPROPERTY()
+		FOneSkyListProjectGroupLanguagesResponseMeta meta;
+
+	UPROPERTY()
+		TArray<FOneSkyListProjectGroupLanguagesResponseDataItem> data;
+
+	/** Default constructor. */
+	FOneSkyListProjectGroupLanguagesResponse()
+		: meta()
+		, data()
+	{ }
+
+	/** Creates an uninitialized instance. */
+	FOneSkyListProjectGroupLanguagesResponse(ENoInit) { }
+};
+
+// end LIST PROJECT GROUP LANGUAGES
+
+
+// LIST PROJECTS IN GROUP
+
+USTRUCT()
+struct FOneSkyListProjectsInGroupResponseMeta
+{
+	GENERATED_USTRUCT_BODY()
+
+		UPROPERTY()
+		int32 status;
+
+	UPROPERTY()
+		int32 record_count;
+
+	/** Default constructor. */
+	FOneSkyListProjectsInGroupResponseMeta()
+		: status(-1)
+		, record_count(-1)
+	{ }
+
+	/** Creates an uninitialized instance. */
+	FOneSkyListProjectsInGroupResponseMeta(ENoInit) { }
+};
+
+USTRUCT()
+struct FOneSkyListProjectsInGroupResponseDataItem
+{
+	GENERATED_USTRUCT_BODY()
+
+	UPROPERTY()
+		int32 id;
+
+	UPROPERTY()
+		FString name;
+
+	/** Default constructor. */
+	FOneSkyListProjectsInGroupResponseDataItem()
+		: id(-1)
+		, name(TEXT(""))
+	{ }
+
+	/** Creates an uninitialized instance. */
+	FOneSkyListProjectsInGroupResponseDataItem(ENoInit) { }
+};
+
+/**
+* Response from a List Projects in Group query on OneSky
+* //[REDACTED_PROFILE_URL]/api-documentation-platform/blob/master/resources/project.md
+*/
+USTRUCT()
+struct FOneSkyListProjectsInGroupResponse
+{
+	GENERATED_USTRUCT_BODY()
+
+	UPROPERTY()
+	FOneSkyListProjectsInGroupResponseMeta meta;
+
+	UPROPERTY()
+		TArray<FOneSkyListProjectsInGroupResponseDataItem> data;
+
+	/** Default constructor. */
+	FOneSkyListProjectsInGroupResponse()
+		: meta()
+		, data()
+	{ }
+
+	/** Creates an uninitialized instance. */
+	FOneSkyListProjectsInGroupResponse(ENoInit) { }
+};
+
+// end LIST PROJECTS IN GROUP
+
+
+// SHOW PROJECT 
+
+USTRUCT()
+struct FOneSkyShowProjectResponseMeta
+{
+	GENERATED_USTRUCT_BODY()
+
+		UPROPERTY()
+		int32 status;
+
+	/** Default constructor. */
+	FOneSkyShowProjectResponseMeta()
+		: status(-1)
+	{ }
+
+	/** Creates an uninitialized instance. */
+	FOneSkyShowProjectResponseMeta(ENoInit) { }
+};
+
+
+USTRUCT()
+struct FOneSkyShowProjectResponseProjectType
+{
+	GENERATED_USTRUCT_BODY()
+
+		UPROPERTY()
+		FString code;
+
+	UPROPERTY()
+		FString name;
+
+	/** Default constructor. */
+	FOneSkyShowProjectResponseProjectType()
+		: code("")
+		, name("")
+	{ }
+
+	/** Creates an uninitialized instance. */
+	FOneSkyShowProjectResponseProjectType(ENoInit) { }
+};
+
+USTRUCT()
+struct FOneSkyShowProjectResponseData
+{
+	GENERATED_USTRUCT_BODY()
+
+	UPROPERTY()
+		int32 id;
+
+	UPROPERTY()
+		FString name;
+
+	UPROPERTY()
+		FString description;
+
+	UPROPERTY()
+		FOneSkyShowProjectResponseProjectType project_type;
+
+	UPROPERTY()
+		int32 string_count;
+
+	UPROPERTY()
+		int32 word_count;
+
+	/** Default constructor. */
+	FOneSkyShowProjectResponseData()
+		: id(-1)
+		, name(TEXT(""))
+		, description("")
+		, string_count(-1)
+		, word_count(-1)
+	{ }
+
+	/** Creates an uninitialized instance. */
+	FOneSkyShowProjectResponseData(ENoInit) { }
+};
+
+/**
+* Response from a Show Project query on OneSky
+* //[REDACTED_PROFILE_URL]/api-documentation-platform/blob/master/resources/project.md
+*/
+USTRUCT()
+struct FOneSkyShowProjectResponse
+{
+	GENERATED_USTRUCT_BODY()
+
+		UPROPERTY()
+		FOneSkyShowProjectResponseMeta meta;
+
+	UPROPERTY()
+		FOneSkyShowProjectResponseData data;
+
+	/** Default constructor. */
+	FOneSkyShowProjectResponse()
+		: meta()
+		, data()
+	{ }
+
+	/** Creates an uninitialized instance. */
+	FOneSkyShowProjectResponse(ENoInit) { }
+};
+
+// end SHOW PROJECT 
+
+
+// CREATE PROJECT 
+
+USTRUCT()
+struct FOneSkyCreateProjectResponseMeta
+{
+	GENERATED_USTRUCT_BODY()
+
+	UPROPERTY()
+		int32 status;
+
+	/** Default constructor. */
+	FOneSkyCreateProjectResponseMeta()
+		: status(-1)
+	{ }
+
+	/** Creates an uninitialized instance. */
+	FOneSkyCreateProjectResponseMeta(ENoInit) { }
+};
+
+
+USTRUCT()
+struct FOneSkyCreateProjectResponseProjectType
+{
+	GENERATED_USTRUCT_BODY()
+
+	UPROPERTY()
+		FString code;
+
+	UPROPERTY()
+		FString name;
+
+	/** Default constructor. */
+	FOneSkyCreateProjectResponseProjectType()
+		: code("")
+		, name("")
+	{ }
+
+	/** Creates an uninitialized instance. */
+	FOneSkyCreateProjectResponseProjectType(ENoInit) { }
+};
+
+USTRUCT()
+struct FOneSkyCreateProjectResponseData
+{
+	GENERATED_USTRUCT_BODY()
+
+	UPROPERTY()
+		int32 id;
+
+	UPROPERTY()
+		FString name;
+
+	UPROPERTY()
+		FString description;
+
+	UPROPERTY()
+		FOneSkyShowProjectResponseProjectType project_type;
+
+	/** Default constructor. */
+	FOneSkyCreateProjectResponseData()
+		: id(-1)
+		, name(TEXT(""))
+		, description("")
+	{ }
+
+	/** Creates an uninitialized instance. */
+	FOneSkyCreateProjectResponseData(ENoInit) { }
+};
+
+/**
+* Response from a Create Project query on OneSky
+* //[REDACTED_PROFILE_URL]/api-documentation-platform/blob/master/resources/project.md
+*/
+USTRUCT()
+struct FOneSkyCreateProjectResponse
+{
+	GENERATED_USTRUCT_BODY()
+
+	UPROPERTY()
+	FOneSkyCreateProjectResponseMeta meta;
+
+	UPROPERTY()
+	FOneSkyCreateProjectResponseData data;
+
+	/** Default constructor. */
+	FOneSkyCreateProjectResponse()
+		: meta()
+		, data()
+	{ }
+
+	/** Creates an uninitialized instance. */
+	FOneSkyCreateProjectResponse(ENoInit) { }
+};
+
+// end CREATE PROJECT 
+
+
+// LIST PROJECT LANGUAGES
+
+USTRUCT()
+struct FOneSkyListProjectLanguagesResponseMeta
+{
+	GENERATED_USTRUCT_BODY()
+
+	UPROPERTY()
+		int32 status;
+
+	UPROPERTY()
+		int32 record_count;
+
+	/** Default constructor. */
+	FOneSkyListProjectLanguagesResponseMeta()
+		: status(-1)
+		, record_count(-1)
+	{ }
+
+	/** Creates an uninitialized instance. */
+	FOneSkyListProjectLanguagesResponseMeta(ENoInit) { }
+};
+
+USTRUCT()
+struct FOneSkyListProjectLanguagesResponseDataItem
+{
+	GENERATED_USTRUCT_BODY()
+
+		UPROPERTY()
+		FString code;
+
+	UPROPERTY()
+		FString english_name;
+
+	UPROPERTY()
+		FString local_name;
+
+	UPROPERTY()
+		FString locale;
+
+	UPROPERTY()
+		FString region;
+
+	UPROPERTY()
+		bool is_base_language;
+
+	UPROPERTY()
+		bool is_ready_to_publish;
+
+	UPROPERTY()
+		FString translation_progress;
+
+	/** Default constructor. */
+	FOneSkyListProjectLanguagesResponseDataItem()
+		: code("")
+		, english_name("")
+		, local_name("")
+		, locale("")
+		, region("")
+		, is_base_language(false)
+		, is_ready_to_publish(false)
+		, translation_progress("")
+	{ }
+
+	/** Creates an uninitialized instance. */
+	FOneSkyListProjectLanguagesResponseDataItem(ENoInit) { }
+};
+
+/**
+* Response from a List Project Languages query on OneSky
+* //[REDACTED_PROFILE_URL]/api-documentation-platform/blob/master/resources/project.md
+*/
+USTRUCT()
+struct FOneSkyListProjectLanguagesResponse
+{
+	GENERATED_USTRUCT_BODY()
+
+	UPROPERTY()
+		FOneSkyListProjectLanguagesResponseMeta meta;
+
+	UPROPERTY()
+		TArray<FOneSkyListProjectLanguagesResponseDataItem> data;
+
+	/** Default constructor. */
+	FOneSkyListProjectLanguagesResponse()
+		: meta()
+		, data()
+	{ }
+
+	/** Creates an uninitialized instance. */
+	FOneSkyListProjectLanguagesResponse(ENoInit) { }
+};
+
+// end LIST PROJECT LANGUAGES
+
+
+// TRANSLATION STATUS
+
+USTRUCT()
+struct FOneSkyTranslationStatusResponseMeta
+{
+	GENERATED_USTRUCT_BODY()
+
+		UPROPERTY()
+		int32 status;
+
+	/** Default constructor. */
+	FOneSkyTranslationStatusResponseMeta()
+		: status(-1)
+	{ }
+
+	/** Creates an uninitialized instance. */
+	FOneSkyTranslationStatusResponseMeta(ENoInit) { }
+};
+
+
+USTRUCT()
+struct FOneSkyTranslationStatusResponseLocale
+{
+	GENERATED_USTRUCT_BODY()
+
+	UPROPERTY()
+		FString code;
+
+	UPROPERTY()
+		FString english_name;
+
+	UPROPERTY()
+		FString local_name;
+
+	UPROPERTY()
+		FString locale;
+
+	UPROPERTY()
+		FString region;
+
+	/** Default constructor. */
+	FOneSkyTranslationStatusResponseLocale()
+		: code("")
+		, english_name("")
+		, local_name("")
+		, locale("")
+		, region("")
+	{ }
+
+	/** Creates an uninitialized instance. */
+	FOneSkyTranslationStatusResponseLocale(ENoInit) { }
+};
+
+USTRUCT()
+struct FOneSkyTranslationStatusResponseData
+{
+	GENERATED_USTRUCT_BODY()
+
+	UPROPERTY()
+		FString file_name;
+
+	UPROPERTY()
+		FOneSkyTranslationStatusResponseLocale locale;
+
+	UPROPERTY()
+		FString progress;
+
+	UPROPERTY()
+		int32 string_count;
+
+	UPROPERTY()
+		int32 word_count;
+
+	/** Default constructor. */
+	FOneSkyTranslationStatusResponseData()
+		: file_name("")
+		, locale()
+		, progress("")
+		, string_count(-1)
+		, word_count(-1)
+	{ }
+
+	/** Creates an uninitialized instance. */
+	FOneSkyTranslationStatusResponseData(ENoInit) { }
+};
+
+/**
+* Response from a Translation Status query on OneSky
+* //[REDACTED_PROFILE_URL]/api-documentation-platform/blob/master/resources/translation.md#status---translations-status
+*/
+USTRUCT()
+struct FOneSkyTranslationStatusResponse
+{
+	GENERATED_USTRUCT_BODY()
+
+	UPROPERTY()
+		FOneSkyTranslationStatusResponseMeta meta;
+
+	UPROPERTY()
+		FOneSkyTranslationStatusResponseData data;
+
+	/** Default constructor. */
+	FOneSkyTranslationStatusResponse()
+		: meta()
+		, data()
+	{ }
+
+	/** Creates an uninitialized instance. */
+	FOneSkyTranslationStatusResponse(ENoInit) { }
+};
+
+// end TRANSLATION STATUS
+
+
+// LIST UPLOADED FILES
+
+USTRUCT()
+struct FOneSkyListUploadedFilesResponseMeta
+{
+	GENERATED_USTRUCT_BODY()
+
+		UPROPERTY()
+		int32 status;
+
+	UPROPERTY()
+		int32 record_count;
+
+	UPROPERTY()
+		int32 page_count;
+
+	UPROPERTY()
+		FString next_page;
+
+	UPROPERTY()
+		FString prev_page;
+
+	UPROPERTY()
+		FString first_page;
+
+	UPROPERTY()
+		FString last_page;
+
+	/** Default constructor. */
+	FOneSkyListUploadedFilesResponseMeta()
+		: status(-1)
+		, record_count(-1)
+		, page_count(-1)
+		, next_page(TEXT(""))
+		, prev_page(TEXT(""))
+		, first_page(TEXT(""))
+		, last_page(TEXT(""))
+	{ }
+
+	/** Creates an uninitialized instance. */
+	FOneSkyListUploadedFilesResponseMeta(ENoInit) { }
+};
+
+USTRUCT()
+struct FOneSkyListUploadedFilesResponseLastImport
+{
+	GENERATED_USTRUCT_BODY()
+
+		UPROPERTY()
+		int32 id;
+
+	UPROPERTY()
+		FString status;
+
+	/** Default constructor. */
+	FOneSkyListUploadedFilesResponseLastImport()
+		: id(-1)
+		, status("")
+	{ }
+
+	/** Creates an uninitialized instance. */
+	FOneSkyListUploadedFilesResponseLastImport(ENoInit) { }
+};
+
+
+USTRUCT()
+struct FOneSkyListUploadedFilesResponseDataItem
+{
+	GENERATED_USTRUCT_BODY()
+
+	UPROPERTY()
+		FString file_name;
+
+	UPROPERTY()
+		int32 string_count;
+
+	UPROPERTY()
+		FOneSkyListUploadedFilesResponseLastImport last_import;
+
+	UPROPERTY()
+		FString uploaded_at;
+
+	UPROPERTY()
+		int32 uploaded_at_timestamp;
+
+	/** Default constructor. */
+	FOneSkyListUploadedFilesResponseDataItem()
+		: file_name("")
+		, string_count(-1)
+		, last_import()
+		, uploaded_at("")
+		, uploaded_at_timestamp(-1)
+	{ }
+
+	/** Creates an uninitialized instance. */
+	FOneSkyListUploadedFilesResponseDataItem(ENoInit) { }
+};
+
+/**
+* Response from a List Uploaded Files query on OneSky
+* //[REDACTED_PROFILE_URL]/api-documentation-platform/blob/master/resources/file.md
+*/
+USTRUCT()
+struct FOneSkyListUploadedFilesResponse
+{
+	GENERATED_USTRUCT_BODY()
+
+		UPROPERTY()
+		FOneSkyListUploadedFilesResponseMeta meta;
+
+	UPROPERTY()
+		TArray<FOneSkyListUploadedFilesResponseDataItem> data;
+
+	/** Default constructor. */
+	FOneSkyListUploadedFilesResponse()
+		: meta()
+		, data()
+	{ }
+
+	/** Creates an uninitialized instance. */
+	FOneSkyListUploadedFilesResponse(ENoInit) { }
+};
+
+// end LIST UPLOADED FILES
+
+// UPLOAD FILE
+
+USTRUCT()
+struct FOneSkyUploadFileResponseMeta
+{
+	GENERATED_USTRUCT_BODY()
+
+	UPROPERTY()
+		int32 status;
+
+	/** Default constructor. */
+	FOneSkyUploadFileResponseMeta()
+		: status(-1)
+	{ }
+
+	/** Creates an uninitialized instance. */
+	FOneSkyUploadFileResponseMeta(ENoInit) { }
+};
+
+
+USTRUCT()
+struct FOneSkyUploadFileResponseLanguage
+{
+	GENERATED_USTRUCT_BODY()
+
+		UPROPERTY()
+		FString code;
+
+	UPROPERTY()
+		FString english_name;
+
+	UPROPERTY()
+		FString local_name;
+
+	UPROPERTY()
+		FString locale;
+
+	UPROPERTY()
+		FString region;
+
+	/** Default constructor. */
+	FOneSkyUploadFileResponseLanguage()
+		: code("")
+		, english_name("")
+		, local_name("")
+		, locale("")
+		, region("")
+	{ }
+
+	/** Creates an uninitialized instance. */
+	FOneSkyUploadFileResponseLanguage(ENoInit) { }
+};
+
+USTRUCT()
+struct FOneSkyUploadedFileResponseImport
+{
+	GENERATED_USTRUCT_BODY()
+
+	UPROPERTY()
+		int32 id;
+
+	UPROPERTY()
+		FString created_at;
+
+	UPROPERTY()
+		int32 created_at_timestamp;
+
+	/** Default constructor. */
+	FOneSkyUploadedFileResponseImport()
+		: id(-1)
+		, created_at("")
+		, created_at_timestamp(-1)
+	{ }
+
+	/** Creates an uninitialized instance. */
+	FOneSkyUploadedFileResponseImport(ENoInit) { }
+};
+
+USTRUCT()
+struct FOneSkyUploadFileResponseData
+{
+	GENERATED_USTRUCT_BODY()
+
+		UPROPERTY()
+		FString name;
+
+	UPROPERTY()
+		FString format;
+
+	UPROPERTY()
+		FOneSkyUploadFileResponseLanguage language;
+
+	UPROPERTY()
+		FOneSkyUploadedFileResponseImport import;
+
+	/** Default constructor. */
+	FOneSkyUploadFileResponseData()
+		: name("")
+		, format("")
+		, language()
+		, import()
+	{ }
+
+	/** Creates an uninitialized instance. */
+	FOneSkyUploadFileResponseData(ENoInit) { }
+};
+
+/**
+* Response from an Upload File post to OneSky
+* //[REDACTED_PROFILE_URL]/api-documentation-platform/blob/master/resources/file.md
+*/
+USTRUCT()
+struct FOneSkyUploadFileResponse
+{
+	GENERATED_USTRUCT_BODY()
+
+	UPROPERTY()
+	FOneSkyUploadFileResponseMeta meta;
+
+	UPROPERTY()
+		FOneSkyUploadFileResponseData data;
+
+	/** Default constructor. */
+	FOneSkyUploadFileResponse()
+		: meta()
+		, data()
+	{ }
+
+	/** Creates an uninitialized instance. */
+	FOneSkyUploadFileResponse(ENoInit) { }
+};
+
+// end UPLOAD FILE
+
+
+// LIST PHRASE COLLECTIONS
+
+USTRUCT()
+struct FOneSkyListPhraseCollectionsResponseMeta
+{
+	GENERATED_USTRUCT_BODY()
+
+		UPROPERTY()
+		int32 status;
+
+	UPROPERTY()
+		int32 record_count;
+
+	UPROPERTY()
+		int32 page_count;
+
+	UPROPERTY()
+		FString next_page;
+
+	UPROPERTY()
+		FString prev_page;
+
+	UPROPERTY()
+		FString first_page;
+
+	UPROPERTY()
+		FString last_page;
+
+	/** Default constructor. */
+	FOneSkyListPhraseCollectionsResponseMeta()
+		: status(-1)
+		, record_count(-1)
+		, page_count(-1)
+		, next_page(TEXT(""))
+		, prev_page(TEXT(""))
+		, first_page(TEXT(""))
+		, last_page(TEXT(""))
+	{ }
+
+	/** Creates an uninitialized instance. */
+	FOneSkyListPhraseCollectionsResponseMeta(ENoInit) { }
+};
+
+USTRUCT()
+struct FOneSkyListPhraseCollectionsResponseDataItem
+{
+	GENERATED_USTRUCT_BODY()
+
+		UPROPERTY()
+		FString key;
+
+	UPROPERTY()
+		FString created_at;
+
+	UPROPERTY()
+		int32 created_at_timestamp;
+
+	/** Default constructor. */
+	FOneSkyListPhraseCollectionsResponseDataItem()
+		: key("")
+		, created_at("")
+		, created_at_timestamp(-1)
+	{ }
+
+	/** Creates an uninitialized instance. */
+	FOneSkyListPhraseCollectionsResponseDataItem(ENoInit) { }
+};
+
+/**
+* Response from a List Uploaded Files query on OneSky
+* //[REDACTED_PROFILE_URL]/api-documentation-platform/blob/master/resources/file.md
+*/
+USTRUCT()
+struct FOneSkyListPhraseCollectionsResponse
+{
+	GENERATED_USTRUCT_BODY()
+
+		UPROPERTY()
+		FOneSkyListPhraseCollectionsResponseMeta meta;
+
+	UPROPERTY()
+		TArray<FOneSkyListPhraseCollectionsResponseDataItem> data;
+
+	/** Default constructor. */
+	FOneSkyListPhraseCollectionsResponse()
+		: meta()
+		, data()
+	{ }
+
+	/** Creates an uninitialized instance. */
+	FOneSkyListPhraseCollectionsResponse(ENoInit) { }
+};
+
+// end LIST PHRASE COLLECTIONS
\ No newline at end of file
diff --git a/Engine/Plugins/Developer/OneSkyLocalizationService/Source/OneSkyLocalizationService/Private/OneSkyLocalizationServiceRevision.cpp b/Engine/Plugins/Developer/OneSkyLocalizationService/Source/OneSkyLocalizationService/Private/OneSkyLocalizationServiceRevision.cpp
new file mode 100644
--- /dev/null
+++ b/Engine/Plugins/Developer/OneSkyLocalizationService/Source/OneSkyLocalizationService/Private/OneSkyLocalizationServiceRevision.cpp
@@ -0,0 +1,32 @@
+// Copyright 1998-2015 Epic Games, Inc. All Rights Reserved.
+
+#include "OneSkyLocalizationServicePrivatePCH.h"
+#include "OneSkyLocalizationServiceRevision.h"
+#include "OneSkyLocalizationServiceModule.h"
+#include "OneSkyLocalizationServiceProvider.h"
+#include "OneSkyConnection.h"
+#include "MessageLog.h"
+
+#define LOCTEXT_NAMESPACE "OneSkyLocalizationService"
+
+int32 FOneSkyLocalizationServiceRevision::GetRevisionNumber() const
+{
+	return RevisionNumber;
+}
+
+const FString& FOneSkyLocalizationServiceRevision::GetRevision() const
+{
+	return Revision;
+}
+
+const FString& FOneSkyLocalizationServiceRevision::GetUserName() const
+{
+	return UserName;
+}
+
+const FDateTime& FOneSkyLocalizationServiceRevision::GetDate() const
+{
+	return Date;
+}
+
+#undef LOCTEXT_NAMESPACE
\ No newline at end of file
diff --git a/Engine/Plugins/Developer/OneSkyLocalizationService/Source/OneSkyLocalizationService/Private/OneSkyLocalizationServiceRevision.h b/Engine/Plugins/Developer/OneSkyLocalizationService/Source/OneSkyLocalizationService/Private/OneSkyLocalizationServiceRevision.h
new file mode 100644
--- /dev/null
+++ b/Engine/Plugins/Developer/OneSkyLocalizationService/Source/OneSkyLocalizationService/Private/OneSkyLocalizationServiceRevision.h
@@ -0,0 +1,35 @@
+// Copyright 1998-2015 Epic Games, Inc. All Rights Reserved.
+
+#pragma once
+
+#include "Developer/LocalizationService/Public/ILocalizationServiceRevision.h"
+
+class FOneSkyLocalizationServiceRevision : public ILocalizationServiceRevision, public TSharedFromThis<FOneSkyLocalizationServiceRevision, ESPMode::ThreadSafe>
+{
+public:
+	FOneSkyLocalizationServiceRevision()
+		: RevisionNumber(0)
+		, Date(0)
+	{
+	}
+
+	/** ILocalizationServiceRevision interface */
+	virtual int32 GetRevisionNumber() const override;
+	virtual const FString& GetRevision() const override;
+	virtual const FString& GetUserName() const override;
+	virtual const FDateTime& GetDate() const override;
+
+public:
+
+	/** The revision number of this file */
+	int32 RevisionNumber;
+
+	/** The revision to display to the user */
+	FString Revision;
+
+	/** THe user that made the change for this revision */
+	FString UserName;
+
+	/** The date of this revision */
+	FDateTime Date;
+};
\ No newline at end of file
diff --git a/Engine/Plugins/Developer/OneSkyLocalizationService/Source/OneSkyLocalizationService/Private/OneSkyLocalizationServiceSettings.cpp b/Engine/Plugins/Developer/OneSkyLocalizationService/Source/OneSkyLocalizationService/Private/OneSkyLocalizationServiceSettings.cpp
new file mode 100644
--- /dev/null
+++ b/Engine/Plugins/Developer/OneSkyLocalizationService/Source/OneSkyLocalizationService/Private/OneSkyLocalizationServiceSettings.cpp
@@ -0,0 +1,199 @@
+// Copyright 1998-2015 Epic Games, Inc. All Rights Reserved.
+
+#include "OneSkyLocalizationServicePrivatePCH.h"
+#include "OneSkyLocalizationServiceSettings.h"
+#include "Developer/LocalizationService/Public/LocalizationServiceHelpers.h"
+#include "OneSkyLocalizationServiceModule.h"
+#include "OneSkyLocalizationServiceProvider.h"
+#include "JsonObject.h"
+#include "JsonReader.h"
+#include "JsonValue.h"
+#include "TextLocalizationManager.h"
+//#include "JsonInternationalizationManifestSerializer.h"
+
+
+namespace OneSkySettingsConstants
+{
+
+/** The section of the ini file we load our settings from */
+static const FString SettingsSection = TEXT("OneSkyLocalizationService.OneSkyLocalizationServiceSettings");
+
+}
+
+UOneSkyLocalizationTargetSettings::UOneSkyLocalizationTargetSettings(const FObjectInitializer& ObjectInitializer)
+	: Super(ObjectInitializer)
+{
+
+}
+
+
+const FString& FOneSkyLocalizationServiceSettings::GetConnectionName() const
+{
+	FScopeLock ScopeLock(&CriticalSection);
+	return ConnectionInfo.Name;
+}
+
+void FOneSkyLocalizationServiceSettings::LoadSettings()
+{
+	FScopeLock ScopeLock(&CriticalSection);
+	const FString& IniFile = LocalizationServiceHelpers::GetSettingsIni();
+
+	GConfig->GetString(*OneSkySettingsConstants::SettingsSection, TEXT("ConnectionName"), ConnectionInfo.Name, IniFile);
+	GConfig->GetString(*OneSkySettingsConstants::SettingsSection, TEXT("ConnectionApiKey"), ConnectionInfo.ApiKey, IniFile);
+	GConfig->GetString(*OneSkySettingsConstants::SettingsSection, TEXT("ConnectionApiSecret"), ConnectionInfo.ApiSecret, IniFile);
+
+	bSaveSecretKey = !ConnectionInfo.ApiSecret.IsEmpty();
+
+	// If connection info was not loaded from .ini, try loading from JSON file
+	if (ConnectionInfo.ApiKey.IsEmpty() && ConnectionInfo.ApiSecret.IsEmpty())
+	{
+		TArray<FString> CredentialsFileNames;
+		// read json config file
+		const FString CredentialsExtension = TEXT("credentials");
+		FString CredentialsFolderPath = FPaths::EnginePluginsDir() / "Developer" / "OneSkyLocalizationService" / "Credentials";
+		// Search for .manifest files 1 directory deep
+		const FString DirectoryToSearchWildcard = FString::Printf(TEXT("%s/*"), *CredentialsFolderPath);
+		TArray<FString> CredentialsFolders;
+		CredentialsFolders.AddUnique(CredentialsFolderPath);
+		IFileManager::Get().FindFiles(CredentialsFolders, *DirectoryToSearchWildcard, false, true);
+		for (FString CredentialsFolder : CredentialsFolders)
+		{
+			// Add root back on
+			CredentialsFolder = CredentialsFolderPath / CredentialsFolder;
+			const FString SubFolderCredentialFileWildcard = FString::Printf(TEXT("%s/*.%s"), *CredentialsFolder, *CredentialsExtension);
+
+			TArray<FString> CurrentFolderCredentialsFileNames;
+			IFileManager::Get().FindFiles(CurrentFolderCredentialsFileNames, *SubFolderCredentialFileWildcard, true, false);
+
+			for (FString& CurrentFolderCredentialsFileName : CurrentFolderCredentialsFileNames)
+			{
+				// Add root back on
+				CurrentFolderCredentialsFileName = CredentialsFolder / CurrentFolderCredentialsFileName;
+				CredentialsFileNames.AddUnique(CurrentFolderCredentialsFileName);
+			}
+		}
+
+		for (FString CredentialsFile : CredentialsFileNames)
+		{
+			//read in file as string
+			FString FileContents;
+			if (!FFileHelper::LoadFileToString(FileContents, *CredentialsFile))
+			{
+				UE_LOG(LogLocalizationService, Log, TEXT("Failed to load OneSky credentials file %s."), *CredentialsFile);
+				continue;
+			}
+
+			//parse as JSON
+			TSharedPtr<FJsonObject> JSONObject;
+
+			TSharedRef< TJsonReader<> > Reader = TJsonReaderFactory<>::Create(FileContents);
+
+			if (!FJsonSerializer::Deserialize(Reader, JSONObject) || !JSONObject.IsValid())
+			{
+				UE_LOG(LogLocalizationService, Log, TEXT("Invalid JSON in OneSky credentials file %s."), *CredentialsFile);
+				continue;
+			}
+
+			TArray<TSharedPtr<FJsonValue> > CredentialsArray = JSONObject->GetArrayField("Credentials");
+
+			for (TSharedPtr<FJsonValue> CredentialsItem : CredentialsArray)
+			{
+				bool bFoundCredentials = false;
+				const TSharedPtr<FJsonObject>& CredentialsObject = CredentialsItem->AsObject();
+				if (CredentialsObject->HasTypedField<EJson::String>("Name"))
+				{
+					ConnectionInfo.Name = CredentialsObject->GetStringField("Name");
+				}
+
+				if (CredentialsObject->HasTypedField<EJson::String>("ApiKey"))
+				{
+					ConnectionInfo.ApiKey = CredentialsObject->GetStringField("ApiKey");
+					bFoundCredentials = true;
+				}
+				else
+				{
+					UE_LOG(LogLocalizationService, Log, TEXT("Credentials file %s is missing ApiKey for connection name %s"), *CredentialsFile, *ConnectionInfo.Name);
+				}
+
+				if (CredentialsObject->HasTypedField<EJson::String>("ApiSecret"))
+				{
+					ConnectionInfo.ApiSecret = CredentialsObject->GetStringField("ApiSecret");
+				}
+				else
+				{
+					UE_LOG(LogLocalizationService, Log, TEXT("Credentials file %s is missing ApiSecret for connection name %s"), *CredentialsFile, *ConnectionInfo.Name);
+				}
+
+				if (bFoundCredentials)
+				{
+					break;
+				}
+			}
+		}
+	}
+}
+
+void FOneSkyLocalizationServiceSettings::SaveSettings() const
+{
+	FScopeLock ScopeLock(&CriticalSection);
+	const FString& IniFile = LocalizationServiceHelpers::GetSettingsIni();
+	GConfig->SetString(*OneSkySettingsConstants::SettingsSection, TEXT("ConnectionName"), *ConnectionInfo.Name, IniFile);
+	GConfig->SetString(*OneSkySettingsConstants::SettingsSection, TEXT("ConnectionApiKey"), *ConnectionInfo.ApiKey, IniFile);
+
+	// Only save secret key if asked to, otherwise store null string
+	FString SecretKeyToSave = ConnectionInfo.ApiSecret;
+	if (!bSaveSecretKey)
+	{
+		SecretKeyToSave = "";
+	}
+
+	GConfig->SetString(*OneSkySettingsConstants::SettingsSection, TEXT("ConnectionApiSecret"), *SecretKeyToSave, IniFile);
+}
+
+FOneSkyConnectionInfo FOneSkyLocalizationServiceSettings::GetConnectionInfo() const
+{
+	check(IsInGameThread());
+	FOneSkyConnectionInfo OutConnectionInfo = ConnectionInfo;
+
+	return OutConnectionInfo;
+}
+
+FOneSkyLocalizationTargetSetting* FOneSkyLocalizationServiceSettings::GetSettingsForTarget(FGuid TargetGuid)
+{
+	FOneSkyLocalizationTargetSetting* Settings = TargetSettingsObject->TargetSettings.FindByPredicate([&](const FOneSkyLocalizationTargetSetting& Setting)
+	{
+		return Setting.TargetGuid == TargetGuid;
+	});
+
+	if (!Settings)
+	{
+		FOneSkyLocalizationTargetSetting NewSettings;
+		NewSettings.TargetGuid = TargetGuid;
+		int32 NewIndex = TargetSettingsObject->TargetSettings.Add(NewSettings);
+		Settings = &(TargetSettingsObject->TargetSettings[NewIndex]);
+	}
+
+	return Settings;
+}
+
+void FOneSkyLocalizationServiceSettings::SetSettingsForTarget(FGuid TargetGuid, FString ProjectGroupName, FString ProjectName, FString FileName)
+{
+	FOneSkyLocalizationTargetSetting* Settings = TargetSettingsObject->TargetSettings.FindByPredicate([&](const FOneSkyLocalizationTargetSetting& Setting)
+	{
+		return Setting.TargetGuid == TargetGuid;
+	});
+
+	if (!Settings)
+	{
+		FOneSkyLocalizationTargetSetting NewSettings;
+		NewSettings.TargetGuid = TargetGuid;
+		int32 NewIndex = TargetSettingsObject->TargetSettings.Add(NewSettings);
+		Settings = &(TargetSettingsObject->TargetSettings[NewIndex]);
+	}
+
+	Settings->OneSkyProjectGroupName = ProjectGroupName;
+	Settings->OneSkyProjectName = ProjectName;
+	Settings->OneSkyFileName = FileName;
+
+	TargetSettingsObject->SaveConfig();
+}
diff --git a/Engine/Plugins/Developer/OneSkyLocalizationService/Source/OneSkyLocalizationService/Private/OneSkyLocalizationServiceSettings.h b/Engine/Plugins/Developer/OneSkyLocalizationService/Source/OneSkyLocalizationService/Private/OneSkyLocalizationServiceSettings.h
new file mode 100644
--- /dev/null
+++ b/Engine/Plugins/Developer/OneSkyLocalizationService/Source/OneSkyLocalizationService/Private/OneSkyLocalizationServiceSettings.h
@@ -0,0 +1,120 @@
+// Copyright 1998-2015 Epic Games, Inc. All Rights Reserved.
+
+#pragma once
+
+#include "OneSkyConnectionInfo.h"
+#include "OneSkyLocalizationServiceSettings.generated.h"
+
+/**
+* Holds the OneSky settings for a Localization Target.
+*/
+USTRUCT()
+struct FOneSkyLocalizationTargetSetting
+{
+	GENERATED_USTRUCT_BODY()
+
+	/**
+	* The GUID of the LocalizationTarget that these OneSky settings are for
+	*/
+	UPROPERTY()
+		FGuid TargetGuid;
+
+	/**
+	* The name of the OneSky Project Group this target belongs to
+	*/
+	UPROPERTY()
+		FString OneSkyProjectGroupName;
+
+	/**
+	* The name of the OneSky Project  this target belongs to
+	*/
+	UPROPERTY()
+		FString OneSkyProjectName;
+
+	/**
+	* The name of the file that corresponds to this target
+	*/
+	UPROPERTY()
+		FString OneSkyFileName;
+
+};
+
+UCLASS(config = Editor)
+class UOneSkyLocalizationTargetSettings : public UObject
+{
+	GENERATED_UCLASS_BODY()
+
+public:
+	/**
+	* The array of settings for each target of this project 
+	*/
+	UPROPERTY(config)
+		TArray<struct FOneSkyLocalizationTargetSetting> TargetSettings;
+
+};
+
+class FOneSkyLocalizationServiceSettings
+{
+public:
+
+	FOneSkyLocalizationServiceSettings() : bSaveSecretKey(false) 
+	{
+		TargetSettingsObject = NewObject<UOneSkyLocalizationTargetSettings>();
+	}
+
+	/** Get the OneSky connection name */
+	const FString& GetConnectionName() const;
+
+	/** Load settings from ini file */
+	void LoadSettings();
+
+	/** Save settings to ini file */
+	void SaveSettings() const;
+
+	/** Get the credentials we use to access the server - only call on the game thread */
+	FOneSkyConnectionInfo GetConnectionInfo() const;
+
+	/** Set the public API key we use to access the server - only call on the game thread */
+	void SetApiKey(FString InApiKey)
+	{
+		check(IsInGameThread());
+		ConnectionInfo.ApiKey = InApiKey;
+	}
+	
+	/** Set the secret API key we use to access the server - only call on the game thread */
+	void SetApiSecret(FString InApiSecret)
+	{
+		check(IsInGameThread());
+		ConnectionInfo.ApiSecret = InApiSecret;
+	}
+
+	/** Set whether or not to save the secret API key (WARNING: saved unencrypted) */
+	void SetSaveSecretKey(bool InSaveSecretKey)
+	{
+		bSaveSecretKey = InSaveSecretKey;
+	}
+
+	/** Get whether or not to save the secret API Key */
+	bool GetSaveSecretKey()
+	{
+		return bSaveSecretKey;
+	}
+
+	FOneSkyLocalizationTargetSetting* GetSettingsForTarget(FGuid TargetGuid);
+
+	void SetSettingsForTarget(FGuid TargetGuid, FString ProjectGroupName, FString ProjectName, FString FileName);
+
+public:
+	/** Object that contains array of OneSky settings for a given targets, used to serialize to .ini file */
+	UOneSkyLocalizationTargetSettings* TargetSettingsObject;
+
+private:
+	/** A critical section for settings access */
+	mutable FCriticalSection CriticalSection;
+
+	/** The credentials we use to access the server */
+	FOneSkyConnectionInfo ConnectionInfo;
+
+	/** Whether or not we remember the secret key for the user (WARNING: saved unencrypted) */
+	bool bSaveSecretKey;
+}; 
\ No newline at end of file
diff --git a/Engine/Plugins/Developer/OneSkyLocalizationService/Source/OneSkyLocalizationService/Private/OneSkyLocalizationServiceState.cpp b/Engine/Plugins/Developer/OneSkyLocalizationService/Source/OneSkyLocalizationService/Private/OneSkyLocalizationServiceState.cpp
new file mode 100644
--- /dev/null
+++ b/Engine/Plugins/Developer/OneSkyLocalizationService/Source/OneSkyLocalizationService/Private/OneSkyLocalizationServiceState.cpp
@@ -0,0 +1,162 @@
+// Copyright 1998-2015 Epic Games, Inc. All Rights Reserved.
+
+#include "OneSkyLocalizationServicePrivatePCH.h"
+#include "OneSkyLocalizationServiceState.h"
+
+#define LOCTEXT_NAMESPACE "OneSkyLocalizationService.State"
+
+int32 FOneSkyLocalizationServiceState::GetHistorySize() const
+{
+	return History.Num();
+}
+
+TSharedPtr<class ILocalizationServiceRevision, ESPMode::ThreadSafe> FOneSkyLocalizationServiceState::GetHistoryItem( int32 HistoryIndex ) const
+{
+	check(History.IsValidIndex(HistoryIndex));
+	return History[HistoryIndex];
+}
+
+FName FOneSkyLocalizationServiceState::GetIconName() const
+{
+	if (!IsCurrent())
+	{
+		return FName("OneSky.NotAtHeadRevision");
+	}
+
+	switch (State)
+	{
+	default:
+	case EOneSkyState::Unknown:
+		return NAME_None;
+	case EOneSkyState::Untranslated:
+		return FName("OneSky.Untranslated");
+	case EOneSkyState::NotAccepted:
+		return FName("OneSky.NotAccepted");
+	case EOneSkyState::NotFinalized:
+		return FName("OneSky.NotFinalized");
+	case EOneSkyState::Finalized:
+		return FName("OneSky.Finalized");
+	case EOneSkyState::Deprecated:
+		return FName("OneSky.Deprecated");
+	}
+}
+
+FName FOneSkyLocalizationServiceState::GetSmallIconName() const
+{
+	if (!IsCurrent())
+	{
+		return FName("OneSky.NotAtHeadRevision_Small");
+	}
+
+	switch (State)
+	{
+	default:
+	case EOneSkyState::Unknown:
+		return NAME_None;
+	case EOneSkyState::Untranslated:
+		return FName("OneSky.Untranslated_Small");
+	case EOneSkyState::NotAccepted:
+		return FName("OneSky.NotAccepted_Small");
+	case EOneSkyState::NotFinalized:
+		return FName("OneSky.NotFinalized_Small");
+	case EOneSkyState::Finalized:
+		return FName("OneSky.Finalized_Small");
+	case EOneSkyState::Deprecated:
+		return FName("OneSky.Deprecated_Small");
+	}
+}
+
+FText FOneSkyLocalizationServiceState::GetDisplayName() const
+{
+
+	if( !IsCurrent() )
+	{
+		return LOCTEXT("NotCurrent", "Not current");
+	}
+
+	switch(State)
+	{
+	default:
+	case EOneSkyState::Unknown:
+		return LOCTEXT("Unknown", "Unknown to OneSky");
+	case EOneSkyState::Untranslated:
+		return LOCTEXT("Untranslated", "Untranslated in OneSky");
+	case EOneSkyState::NotAccepted:
+		return LOCTEXT("NotAccepted", "Translation not accepted in OneSky");
+	case EOneSkyState::NotFinalized:
+		return LOCTEXT("NotFinalized", "Translation accepted in OneSky");
+	case EOneSkyState::Finalized:
+		return LOCTEXT("Finalized", "Translation finalized in OneSky");
+	case EOneSkyState::Deprecated:
+		return LOCTEXT("Deprecated", "Translation deprecated in OneSky");
+	}
+}
+
+FText FOneSkyLocalizationServiceState::GetDisplayTooltip() const
+{
+	if( !IsCurrent() )
+	{
+		return LOCTEXT("NotCurrent_Tooltip", "The text(s) are not at the latest revision");
+	}
+
+	switch(State)
+	{
+	default:
+	case EOneSkyState::Unknown:
+		return LOCTEXT("Unknown_Tooltip", "The text(s) status is unknown to OneSky");
+	case EOneSkyState::Untranslated:
+		return LOCTEXT("Untranslated_Tooltip", "The text(s) is known to OneSky, but there is no translation for the given culture.");
+	case EOneSkyState::NotAccepted:
+		return LOCTEXT("NotAccepted_Tooltip", "The text(s) is know to OneSky and there is a translation for the given culture but it is not accepted.");
+	case EOneSkyState::NotFinalized:
+		return LOCTEXT("NotFinalized_Tooltip", "The text(s) is know to OneSky, and a translation for the given culture has been accepted, but not finalized.");
+	case EOneSkyState::Finalized:
+		return LOCTEXT("Finalized_Tooltip", "The text(s) is know to OneSky, and a translation for the given culture has been accepted and finalized.");
+	case EOneSkyState::Deprecated:
+		return LOCTEXT("Deprecated_Tooltip", "The text(s) is know to OneSky, and a translation existed previously but it has been deprecated.");
+	};
+}
+
+const FString& FOneSkyLocalizationServiceState::GetSourceString() const
+{
+	return TranslationId.Source;
+}
+
+const FCulturePtr FOneSkyLocalizationServiceState::GetCulture() const
+{
+	return TranslationId.Culture;
+}
+
+
+const FDateTime& FOneSkyLocalizationServiceState::GetTimeStamp() const
+{
+	return TimeStamp;
+}
+
+bool FOneSkyLocalizationServiceState::IsCurrent() const
+{
+	return LocalRevNumber == OneSkyLatestRevNumber;
+}
+
+bool FOneSkyLocalizationServiceState::IsKnownToLocalizationService() const
+{
+	return State != EOneSkyState::Unknown;
+}
+
+bool FOneSkyLocalizationServiceState::CanEdit() const
+{
+	// TODO: Check if this user can add/accept translations, or change finalized state
+	return true;
+}
+
+bool FOneSkyLocalizationServiceState::IsUnknown() const
+{
+	return State == EOneSkyState::Unknown;
+}
+
+bool FOneSkyLocalizationServiceState::IsModified() const
+{
+	return bModifed;
+}
+
+#undef LOCTEXT_NAMESPACE
diff --git a/Engine/Plugins/Developer/OneSkyLocalizationService/Source/OneSkyLocalizationService/Private/OneSkyLocalizationServiceState.h b/Engine/Plugins/Developer/OneSkyLocalizationService/Source/OneSkyLocalizationService/Private/OneSkyLocalizationServiceState.h
new file mode 100644
--- /dev/null
+++ b/Engine/Plugins/Developer/OneSkyLocalizationService/Source/OneSkyLocalizationService/Private/OneSkyLocalizationServiceState.h
@@ -0,0 +1,112 @@
+// Copyright 1998-2015 Epic Games, Inc. All Rights Reserved.
+
+#pragma once
+
+#include "Developer/LocalizationService/Public/ILocalizationServiceState.h"
+#include "OneSkyLocalizationServiceRevision.h"
+
+namespace EOneSkyState
+{
+	enum Type
+	{
+		/** Text unknown to OneSky */
+		Unknown = 0,
+
+		/** OneSky has an entry for this text, but no translation for this culture */
+		Untranslated = 1,
+
+		/** One sky has an entry for this text and a translation, but the translation has not been accepted*/
+		NotAccepted = 2,
+
+		/** One sky has an entry for this text and an accepted translation, but the translation is not finalized*/
+		NotFinalized = 3,
+
+		/** One sky has an entry for this text and a finalized translation */
+		Finalized = 4,
+
+		/** One sky has an entry for this text, and a translation, but that translation is deprecated */
+		Deprecated = 5,
+	};
+}
+
+class FOneSkyLocalizationServiceState : public ILocalizationServiceState, public TSharedFromThis<FOneSkyLocalizationServiceState, ESPMode::ThreadSafe>
+{
+public:
+	FOneSkyLocalizationServiceState(const FLocalizationServiceTranslationIdentifier InTranslationId, EOneSkyState::Type InState = EOneSkyState::Unknown)
+		: TranslationId(InTranslationId)
+		, State(InState)
+		, OneSkyLatestRevNumber(INVALID_REVISION)
+		, LocalRevNumber(INVALID_REVISION)
+		, bModifed(false)
+		, TimeStamp(0)
+	{
+	}
+
+	/** ILocalizationServiceState interface */
+	virtual int32 GetHistorySize() const override;
+	virtual TSharedPtr<class ILocalizationServiceRevision, ESPMode::ThreadSafe> GetHistoryItem( int32 HistoryIndex ) const override;
+	virtual FName GetIconName() const override;
+	virtual FName GetSmallIconName() const override;
+	virtual FText GetDisplayName() const override;
+	virtual FText GetDisplayTooltip() const override;
+	virtual const FString& GetSourceString() const override;
+	virtual const FCulturePtr GetCulture() const override;
+	virtual const FDateTime& GetTimeStamp() const override;
+	virtual bool IsCurrent() const override;
+	virtual bool IsKnownToLocalizationService() const override;
+	virtual bool IsUnknown() const override;
+	virtual bool CanEdit() const override;
+	virtual bool IsModified() const override;
+	virtual const FLocalizationServiceTranslationIdentifier& GetTranslationIdentifier() const override
+	{
+		return TranslationId;
+	}
+
+	/** Get the state of a file */
+	EOneSkyState::Type GetState() const
+	{
+		return State;
+	}
+
+	/** Set the state of the translation */
+	void SetState( EOneSkyState::Type InState )
+	{
+		State = InState;
+	}
+
+	/** Update the translation */
+	void SetTranslation (const FString& InTranslation)
+	{
+		Translation = InTranslation;
+	}
+
+	virtual const FString& GetTranslationString() const override
+	{
+		return Translation;
+	}
+
+public:
+	/** History of the item, if any */
+	TArray< TSharedRef<FOneSkyLocalizationServiceRevision, ESPMode::ThreadSafe> > History;
+
+	/** Translation identifier */
+	FLocalizationServiceTranslationIdentifier TranslationId;
+
+	/** Current translation */
+	FString Translation;
+
+	/** Status of the file */
+	EOneSkyState::Type State;
+
+	/** Latest revision number of the text in OneSky */
+	int OneSkyLatestRevNumber;
+
+	/** Latest rev number at which the text was synced to before being edited */
+	int LocalRevNumber;
+
+	/** Modified from depot version */
+	bool bModifed;
+
+	/** The timestamp of the last update */
+	FDateTime TimeStamp;
+};
diff --git a/Engine/Source/Developer/LocalizationService/Private/DefaultLocalizationServiceProvider.cpp b/Engine/Source/Developer/LocalizationService/Private/DefaultLocalizationServiceProvider.cpp
--- a/Engine/Source/Developer/LocalizationService/Private/DefaultLocalizationServiceProvider.cpp
+++ b/Engine/Source/Developer/LocalizationService/Private/DefaultLocalizationServiceProvider.cpp
@@ -67,14 +67,16 @@ void FDefaultLocalizationServiceProvider::Tick()
 
 }
 
+#if LOCALIZATION_SERVICES_WITH_SLATE
 void FDefaultLocalizationServiceProvider::CustomizeSettingsDetails(IDetailCategoryBuilder& DetailCategoryBuilder) const
 {
 
 }
 
-void FDefaultLocalizationServiceProvider::CustomizeTargetDetails(IDetailCategoryBuilder& DetailCategoryBuilder, ULocalizationTarget& Target) const
+void FDefaultLocalizationServiceProvider::CustomizeTargetDetails(IDetailCategoryBuilder& DetailCategoryBuilder, const FGuid& TargetGuid) const
 {
 
 }
+#endif //LOCALIZATION_SERVICES_WITH_SLATE
 
 #undef LOCTEXT_NAMESPACE
\ No newline at end of file
diff --git a/Engine/Source/Developer/LocalizationService/Private/DefaultLocalizationServiceProvider.h b/Engine/Source/Developer/LocalizationService/Private/DefaultLocalizationServiceProvider.h
--- a/Engine/Source/Developer/LocalizationService/Private/DefaultLocalizationServiceProvider.h
+++ b/Engine/Source/Developer/LocalizationService/Private/DefaultLocalizationServiceProvider.h
@@ -25,7 +25,7 @@ public:
 	virtual void Tick() override;
 #if LOCALIZATION_SERVICES_WITH_SLATE
 	virtual void CustomizeSettingsDetails(IDetailCategoryBuilder& DetailCategoryBuilder) const override;
-	virtual void CustomizeTargetDetails(IDetailCategoryBuilder& DetailCategoryBuilder, ULocalizationTarget& Target) const override;
+	virtual void CustomizeTargetDetails(IDetailCategoryBuilder& DetailCategoryBuilder, const FGuid& TargetGuid) const override;
 #endif // LOCALIZATION_SERVICES_WITH_SLATE
 };
 
diff --git a/Engine/Source/Developer/LocalizationService/Private/LocalizationServicePrivatePCH.h b/Engine/Source/Developer/LocalizationService/Private/LocalizationServicePrivatePCH.h
--- a/Engine/Source/Developer/LocalizationService/Private/LocalizationServicePrivatePCH.h
+++ b/Engine/Source/Developer/LocalizationService/Private/LocalizationServicePrivatePCH.h
@@ -23,4 +23,12 @@
 #if WITH_UNREAL_DEVELOPER_TOOLS
 #include "MessageLogModule.h"
 #endif
- 
\ No newline at end of file
+ 
+#include "ILocalizationServiceModule.h"
+#include "LocalizationServiceModule.h"
+#include "LocalizationServiceSettings.h"
+#include "ILocalizationServiceProvider.h"
+#include "ILocalizationServiceRevision.h"
+#include "ILocalizationServiceOperation.h" 
+#include "ILocalizationServiceState.h"
+#include "LocalizationServiceOperations.h" 
\ No newline at end of file
diff --git a/Engine/Source/Developer/LocalizationService/Public/ILocalizationServiceProvider.h b/Engine/Source/Developer/LocalizationService/Public/ILocalizationServiceProvider.h
--- a/Engine/Source/Developer/LocalizationService/Public/ILocalizationServiceProvider.h
+++ b/Engine/Source/Developer/LocalizationService/Public/ILocalizationServiceProvider.h
@@ -202,7 +202,7 @@ public:
 	* Create a settings widget for display in the localization target editor.
 	* @returns a widget used to edit target-specific localization service settings.
 	*/
-	virtual void CustomizeTargetDetails(IDetailCategoryBuilder& DetailCategoryBuilder, ULocalizationTarget& Target) const = 0;
+	virtual void CustomizeTargetDetails(IDetailCategoryBuilder& DetailCategoryBuilder, const FGuid& TargetGuid) const = 0;
 #endif //LOCALIZATION_SERVICES_WITH_SLATE
 };
 
diff --git a/Engine/Source/Developer/LocalizationService/Public/ILocalizationServiceState.h b/Engine/Source/Developer/LocalizationService/Public/ILocalizationServiceState.h
--- a/Engine/Source/Developer/LocalizationService/Public/ILocalizationServiceState.h
+++ b/Engine/Source/Developer/LocalizationService/Public/ILocalizationServiceState.h
@@ -12,13 +12,15 @@ struct FLocalizationServiceTranslationIdentifier //: TSharedFromThis<FLocalizati
 	{}
 
 	FLocalizationServiceTranslationIdentifier(FCulturePtr InCulturePtr, const FString& InNamespace, const FString& InSource)
-		: Culture(InCulturePtr)
+		: LocalizationTargetGuid()
+		, Culture(InCulturePtr)
 		, Namespace(InNamespace)
 		, Source(InSource)
 	{
 
 	}
 
+	FGuid LocalizationTargetGuid;
 	FCulturePtr Culture;
 	FString Namespace;
 	FString Source;
