OneSky Plugin and Localization Service Module changes:
- Created the FDownloadLocalizationTargetFile Localization Service Operation, which is handled in the OneSky plugin via the FOneSkyTranslationExportWorker.
- Changed OneSky settings to use the OneSky ID for project instead of project name string (OneSky allows multiple projects to have the same name).
- Removed ProjectGroup setting. Only used when creating a project, should just ask for it in the Create Project UI, when it comes time to make one.

diff --git a/Engine/Plugins/Developer/OneSkyLocalizationService/Source/OneSkyLocalizationService/Private/OneSkyLocalizationServiceModule.cpp b/Engine/Plugins/Developer/OneSkyLocalizationService/Source/OneSkyLocalizationService/Private/OneSkyLocalizationServiceModule.cpp
--- a/Engine/Plugins/Developer/OneSkyLocalizationService/Source/OneSkyLocalizationService/Private/OneSkyLocalizationServiceModule.cpp
+++ b/Engine/Plugins/Developer/OneSkyLocalizationService/Source/OneSkyLocalizationService/Private/OneSkyLocalizationServiceModule.cpp
@@ -35,7 +35,7 @@ void FOneSkyLocalizationServiceModule::StartupModule()
 
 	// Translation API
 	OneSkyLocalizationServiceProvider.RegisterWorker("TranslationStatus", FGetOneSkyLocalizationServiceWorker::CreateStatic(&CreateWorker<FOneSkyTranslationStatusWorker>));
-	OneSkyLocalizationServiceProvider.RegisterWorker("TranslationExport", FGetOneSkyLocalizationServiceWorker::CreateStatic(&CreateWorker<FOneSkyTranslationExportWorker>));
+	OneSkyLocalizationServiceProvider.RegisterWorker("DownloadLocalizationTargetFile", FGetOneSkyLocalizationServiceWorker::CreateStatic(&CreateWorker<FOneSkyTranslationExportWorker>));
 
 	// Files API
 	OneSkyLocalizationServiceProvider.RegisterWorker("ListUploadedFiles", FGetOneSkyLocalizationServiceWorker::CreateStatic(&CreateWorker<FOneSkyListUploadedFilesWorker>));
diff --git a/Engine/Plugins/Developer/OneSkyLocalizationService/Source/OneSkyLocalizationService/Private/OneSkyLocalizationServiceOperations.cpp b/Engine/Plugins/Developer/OneSkyLocalizationService/Source/OneSkyLocalizationService/Private/OneSkyLocalizationServiceOperations.cpp
--- a/Engine/Plugins/Developer/OneSkyLocalizationService/Source/OneSkyLocalizationService/Private/OneSkyLocalizationServiceOperations.cpp
+++ b/Engine/Plugins/Developer/OneSkyLocalizationService/Source/OneSkyLocalizationService/Private/OneSkyLocalizationServiceOperations.cpp
@@ -691,19 +691,31 @@ bool FOneSkyTranslationExportWorker::Execute(class FOneSkyLocalizationServiceCom
 {
 	// Store pointer to command so we can access it in the http request callback
 	Command = &InCommand;
-	TSharedPtr<FOneSkyTranslationExportOperation, ESPMode::ThreadSafe> TranslationExportOp = StaticCastSharedRef<FOneSkyTranslationExportOperation>(InCommand.Operation);
+	TSharedPtr<FDownloadLocalizationTargetFile, ESPMode::ThreadSafe> DownloadLocTargetOp = StaticCastSharedRef<FDownloadLocalizationTargetFile>(InCommand.Operation);
+
+	FGuid InTargetGuid;
+	FString InLocale;
+	FString InRelativeOutputFilePathAndName;
+
+	if (DownloadLocTargetOp.IsValid())
+	{
+		InTargetGuid = DownloadLocTargetOp->GetInTargetGuid();
+		InLocale = FGenericPlatformHttp::UrlEncode(DownloadLocTargetOp->GetInLocale());
+		InRelativeOutputFilePathAndName = DownloadLocTargetOp->GetInRelativeOutputFilePathAndName();
+	}
 
 	int32 InProjectId = -1;
 	FString InSourceFileName = "";
 	FString InExportFileName = "";
-	FString InLocale = "";
 
-	if (TranslationExportOp.IsValid())
+	// OneSky project settings are accessed by their localization target guid. These settings are set in the Localization Dashboard
+	FOneSkyLocalizationTargetSetting* Settings = FOneSkyLocalizationServiceModule::Get().AccessSettings().GetSettingsForTarget(InTargetGuid);
+	if (Settings != nullptr)
 	{
-		InProjectId = TranslationExportOp->GetInProjectId();
-		InSourceFileName = FGenericPlatformHttp::UrlEncode(TranslationExportOp->GetInSourceFileName());
-		InExportFileName = FGenericPlatformHttp::UrlEncode(TranslationExportOp->GetInExportFileName());
-		InLocale = FGenericPlatformHttp::UrlEncode(TranslationExportOp->GetInLocale());
+		ILocalizationServiceProvider& Provider = ILocalizationServiceModule::Get().GetProvider();
+		InProjectId = Settings->OneSkyProjectId;
+		InSourceFileName = FGenericPlatformHttp::UrlEncode(Settings->OneSkyFileName);
+		InExportFileName = FGenericPlatformHttp::UrlEncode(FPaths::ConvertRelativePathToFull(InRelativeOutputFilePathAndName));
 	}
 
 	FString Url = AddAuthenticationParameters(InCommand.ConnectionInfo, "https://platform.api.onesky.io/1/projects/" + FString::FromInt(InProjectId) / "translations");
@@ -746,16 +758,20 @@ void FOneSkyTranslationExportWorker::Query_HttpRequestComplete(FHttpRequestPtr H
 				{
 					if (Command != nullptr)
 					{
-						TSharedPtr<FOneSkyTranslationExportOperation, ESPMode::ThreadSafe> TranslationExportOp = StaticCastSharedRef<FOneSkyTranslationExportOperation>(Command->Operation);
+						TSharedPtr<FDownloadLocalizationTargetFile, ESPMode::ThreadSafe> TranslationExportOp = StaticCastSharedRef<FDownloadLocalizationTargetFile>(Command->Operation);
 						if (TranslationExportOp.IsValid())
 						{
 							// The file is not ready, try again
-							TSharedRef<FOneSkyTranslationExportOperation, ESPMode::ThreadSafe> NewTranslationExportLanguagesOp = ILocalizationServiceOperation::Create<FOneSkyTranslationExportOperation>();
-							NewTranslationExportLanguagesOp->SetInProjectId(TranslationExportOp->GetInProjectId());
-							NewTranslationExportLanguagesOp->SetInSourceFileName(TranslationExportOp->GetInSourceFileName());
+							TSharedRef<FDownloadLocalizationTargetFile, ESPMode::ThreadSafe> NewTranslationExportLanguagesOp = ILocalizationServiceOperation::Create<FDownloadLocalizationTargetFile>();
+							NewTranslationExportLanguagesOp->SetInTargetGuid(TranslationExportOp->GetInTargetGuid());
 							NewTranslationExportLanguagesOp->SetInLocale(TranslationExportOp->GetInLocale());
-							NewTranslationExportLanguagesOp->SetInExportFileName(TranslationExportOp->GetInExportFileName());
-							ILocalizationServiceModule::Get().GetProvider().Execute(NewTranslationExportLanguagesOp, TArray<FLocalizationServiceTranslationIdentifier>(), ELocalizationServiceOperationConcurrency::Asynchronous);
+							NewTranslationExportLanguagesOp->SetInRelativeOutputFilePathAndName(TranslationExportOp->GetInRelativeOutputFilePathAndName());
+							// TODO: Can't spawn a new Worker here, as it tries to access the OneSky Connection Info from a thread that is not the main thread.
+							// Instead spawn in callback?
+							//ILocalizationServiceModule::Get().GetProvider().Execute(NewTranslationExportLanguagesOp, TArray<FLocalizationServiceTranslationIdentifier>(), ELocalizationServiceOperationConcurrency::Asynchronous);
+							// For now, error
+							ErrorText = LOCTEXT("TranslationExportQueryFailedRetryNotImplemented", "Translation Export Query Failed: Retry not yet implemented.");
+							bResult = false;
 						}
 						else
 						{
@@ -778,10 +794,11 @@ void FOneSkyTranslationExportWorker::Query_HttpRequestComplete(FHttpRequestPtr H
 				{
 					if (Command != nullptr)
 					{
-						TSharedPtr<FOneSkyTranslationExportOperation, ESPMode::ThreadSafe> TranslationExportOp = StaticCastSharedRef<FOneSkyTranslationExportOperation>(Command->Operation);
+						TSharedPtr<FDownloadLocalizationTargetFile, ESPMode::ThreadSafe> TranslationExportOp = StaticCastSharedRef<FDownloadLocalizationTargetFile>(Command->Operation);
 						if (TranslationExportOp.IsValid())
 						{
-							FString Filename = TranslationExportOp->GetInExportFileName();
+							// Path is relative to game directory
+							FString Filename = FPaths::ConvertRelativePathToFull(FPaths::GameDir() / TranslationExportOp->GetInRelativeOutputFilePathAndName());
 
 							if (Filename.IsEmpty())
 							{
@@ -815,6 +832,7 @@ void FOneSkyTranslationExportWorker::Query_HttpRequestComplete(FHttpRequestPtr H
 		if (!bResult)
 		{
 			UE_LOG(LogLocalizationService, Warning, TEXT("%s"), *(ErrorText.ToString()));
+			Command->ErrorMessages.Add(ErrorText);
 		}
 	}
 
@@ -1090,3 +1108,4 @@ void FOneSkyListPhraseCollectionsWorker::Query_HttpRequestComplete(FHttpRequestP
 
 #undef LOCTEXT_NAMESPACE
 
+
diff --git a/Engine/Plugins/Developer/OneSkyLocalizationService/Source/OneSkyLocalizationService/Private/OneSkyLocalizationServiceOperations.h b/Engine/Plugins/Developer/OneSkyLocalizationService/Source/OneSkyLocalizationService/Private/OneSkyLocalizationServiceOperations.h
--- a/Engine/Plugins/Developer/OneSkyLocalizationService/Source/OneSkyLocalizationService/Private/OneSkyLocalizationServiceOperations.h
+++ b/Engine/Plugins/Developer/OneSkyLocalizationService/Source/OneSkyLocalizationService/Private/OneSkyLocalizationServiceOperations.h
@@ -621,73 +621,6 @@ public:
 
 // TRANSLATION EXPORT
 
-/**
-* Operation used to list the languages supported by a project in one sky
-*/
-class FOneSkyTranslationExportOperation : public ILocalizationServiceOperation
-{
-public:
-	// ILocalizationServiceOperation interface
-	virtual FName GetName() const override
-	{
-		return "TranslationExport";
-	}
-
-	virtual FText GetInProgressString() const override
-	{
-		return LOCTEXT("OneSkyLocalizationService_TranslationExport", "Exporting translations for a file in a project in OneSky...");
-	}
-
-	FOneSkyTranslationExportOperation() : InProjectId(-1) { }
-
-	int32 GetInProjectId()
-	{
-		return InProjectId;
-	}
-
-	void SetInProjectId(int32 NewProjectGroupId)
-	{
-		InProjectId = NewProjectGroupId;
-	}
-
-	FString GetInSourceFileName()
-	{
-		return InSourceFileName;
-	}
-
-	void SetInSourceFileName(FString NewSourceFileName)
-	{
-		InSourceFileName = NewSourceFileName;
-	}
-
-	FString GetInExportFileName()
-	{
-		return InExportFileName;
-	}
-
-	void SetInExportFileName(FString NewExportFileName)
-	{
-		InExportFileName = NewExportFileName;
-	}
-
-	FString GetInLocale()
-	{
-		return InLocale;
-	}
-
-	void SetInLocale(FString NewLocale)
-	{
-		InLocale = NewLocale;
-	}
-
-protected:
-
-	int32 InProjectId;
-	FString InSourceFileName;
-	FString InExportFileName;
-	FString InLocale;
-};
-
 class FOneSkyTranslationExportWorker : public IOneSkyLocalizationServiceWorker
 {
 public:
diff --git a/Engine/Plugins/Developer/OneSkyLocalizationService/Source/OneSkyLocalizationService/Private/OneSkyLocalizationServiceProvider.cpp b/Engine/Plugins/Developer/OneSkyLocalizationService/Source/OneSkyLocalizationService/Private/OneSkyLocalizationServiceProvider.cpp
--- a/Engine/Plugins/Developer/OneSkyLocalizationService/Source/OneSkyLocalizationService/Private/OneSkyLocalizationServiceProvider.cpp
+++ b/Engine/Plugins/Developer/OneSkyLocalizationService/Source/OneSkyLocalizationService/Private/OneSkyLocalizationServiceProvider.cpp
@@ -26,12 +26,8 @@ static FName ProviderName("OneSky");
 /** Init of connection with source control server */
 void FOneSkyLocalizationServiceProvider::Init(bool bForceConnection)
 {
-	// TODO: Test if connection works
+	// TODO: Test if connection works?
 	bServerAvailable = true;
-
-	// TODO: this is called on startup, delay until we need it
-	// Move to EstablishConnection?
-	RetrieveProjectGroups(ELocalizationServiceOperationConcurrency::Asynchronous);
 }
 
 /** API Specific close the connection with localization provider server*/
@@ -46,8 +42,6 @@ void FOneSkyLocalizationServiceProvider::Close()
 
 	// clear the cache
 	StateCache.Empty();
-	ProjectGroupCache.Empty();
-	ProjectCachePerGroupId.Empty();
 
 	bServerAvailable = false;
 }
@@ -69,103 +63,6 @@ TSharedRef<FOneSkyLocalizationServiceState, ESPMode::ThreadSafe> FOneSkyLocaliza
 	}
 }
 
-bool FOneSkyLocalizationServiceProvider::RetrieveProjectGroups(ELocalizationServiceOperationConcurrency::Type InConcurrency)
-{
-	// TODO: This should do the same thing as GetStateInternal but for project groups
-	// Create a data type to hold the information from a project group
-	// This logic should be moved to a function similar to GetState, but for projects
-	if (!IsEnabled())
-	{
-		return ELocalizationServiceOperationCommandResult::Failed;
-	}
-
-	// LIST PROJECT GROUPS TEST
-	//Execute(ILocalizationServiceOperation::Create<FOneSkyListProjectGroupsOperation>(), TArray<FLocalizationServiceTranslationIdentifier>(), InConcurrency);
-	//TSharedRef<FOneSkyListProjectGroupsOperation, ESPMode::ThreadSafe> ListProjectGroupsOperation = ILocalizationServiceOperation::Create<FOneSkyListProjectGroupsOperation>();
-	//ListProjectGroupsOperation->SetInStartPage(1);
-	//ListProjectGroupsOperation->SetInItemsPerPage(1);
-	//Execute(ListProjectGroupsOperation, TArray<FLocalizationServiceTranslationIdentifier>(), InConcurrency);
-
-	// SHOW PROJECT GROUP TEST
-	//TSharedRef<FOneSkyShowProjectGroupOperation, ESPMode::ThreadSafe> ShowProjectGroupOperation = ILocalizationServiceOperation::Create<FOneSkyShowProjectGroupOperation>();
-	//ShowProjectGroupOperation->SetInProjectGroupId(10222);
-	//Execute(ShowProjectGroupOperation, TArray<FLocalizationServiceTranslationIdentifier>(), InConcurrency);
-
-	// CREATE PROJECT GROUP TEST
-	//TSharedRef<FOneSkyCreateProjectGroupOperation, ESPMode::ThreadSafe> CreateProjectGroupOperation = ILocalizationServiceOperation::Create<FOneSkyCreateProjectGroupOperation>();
-	//CreateProjectGroupOperation->SetInProjectGroupName("test group");
-	//CreateProjectGroupOperation->SetInInProjectGroupBaseLocale("zh-TW");
-	//Execute(CreateProjectGroupOperation, TArray<FLocalizationServiceTranslationIdentifier>(), InConcurrency);
-
-	// LIST PROJECT GROUP LANGUAGES TEST
-	//TSharedRef<FOneSkyListProjectGroupLanguagesOperation, ESPMode::ThreadSafe> ListProjectGroupLanguagesOperation = ILocalizationServiceOperation::Create<FOneSkyListProjectGroupLanguagesOperation>();
-	//ListProjectGroupLanguagesOperation->SetInProjectGroupId(10222);
-	//Execute(ListProjectGroupLanguagesOperation, TArray<FLocalizationServiceTranslationIdentifier>(), InConcurrency);
-
-	// LIST PROJECTS IN GROUP TEST
-	//TSharedRef<FOneSkyListProjectsInGroupOperation, ESPMode::ThreadSafe> ListProjectsInGroupOperation = ILocalizationServiceOperation::Create<FOneSkyListProjectsInGroupOperation>();
-	//ListProjectsInGroupOperation->SetInProjectGroupId(10222);
-	//Execute(ListProjectsInGroupOperation, TArray<FLocalizationServiceTranslationIdentifier>(), InConcurrency);
-
-	// SHOW PROJECT test
-	//TSharedRef<FOneSkyShowProjectOperation, ESPMode::ThreadSafe> ShowProjectOperation = ILocalizationServiceOperation::Create<FOneSkyShowProjectOperation>();
-	//ShowProjectOperation->SetInProjectId(27922);
-	//Execute(ShowProjectOperation, TArray<FLocalizationServiceTranslationIdentifier>(), InConcurrency);
-
-	// CREATE PROJECT test
-	//TSharedRef<FOneSkyCreateProjectOperation, ESPMode::ThreadSafe> CreateProjectOperation = ILocalizationServiceOperation::Create<FOneSkyCreateProjectOperation>();
-	//CreateProjectOperation->SetInProjectGroupId(10222);
-	//CreateProjectOperation->SetInProjectGroupName("test create project");
-	//CreateProjectOperation->SetInInProjectDescription("test project description");
-	//CreateProjectOperation->SetInInProjectType("website");
-	//Execute(CreateProjectOperation, TArray<FLocalizationServiceTranslationIdentifier>(), InConcurrency);
-
-	// LIST PROJECT LANGUAGES TEST
-	//TSharedRef<FOneSkyListProjectLanguagesOperation, ESPMode::ThreadSafe> ListProjectLanguagesOperation = ILocalizationServiceOperation::Create<FOneSkyListProjectLanguagesOperation>();
-	//ListProjectLanguagesOperation->SetInProjectId(27922);
-	//Execute(ListProjectLanguagesOperation, TArray<FLocalizationServiceTranslationIdentifier>(), InConcurrency);
-
-	// TRANSLATION STATUS TEST
-	//TSharedRef<FOneSkyTranslationStatusOperation, ESPMode::ThreadSafe> TranslationStatusOperation = ILocalizationServiceOperation::Create<FOneSkyTranslationStatusOperation>();
-	//TranslationStatusOperation->SetInProjectId(27922);
-	//TranslationStatusOperation->SetInFileName("en.po");
-	//TranslationStatusOperation->SetInLocale("ko");
-	//Execute(TranslationStatusOperation, TArray<FLocalizationServiceTranslationIdentifier>(), InConcurrency);
-
-	// TRANSLATION EXPORT TEST
-	//TSharedRef<FOneSkyTranslationExportOperation, ESPMode::ThreadSafe> TranslationExportLanguagesOperation = ILocalizationServiceOperation::Create<FOneSkyTranslationExportOperation>();
-	//TranslationExportLanguagesOperation->SetInProjectId(27922);
-	//TranslationExportLanguagesOperation->SetInSourceFileName("en.po");
-	//TranslationExportLanguagesOperation->SetInExportFileName("D:\\out-ko.po");
-	//TranslationExportLanguagesOperation->SetInLocale("ko");
-	//Execute(TranslationExportLanguagesOperation, TArray<FLocalizationServiceTranslationIdentifier>(), InConcurrency);
-
-	// LIST UPLOADED FILES
-	//TSharedRef<FOneSkyListUploadedFilesOperation, ESPMode::ThreadSafe> ListUploadedFilesOperation = ILocalizationServiceOperation::Create<FOneSkyListUploadedFilesOperation>();
-	//ListUploadedFilesOperation->SetInProjectId(27922);
-	//ListUploadedFilesOperation->SetInItemsPerPage(1);
-	//ListUploadedFilesOperation->SetInStartPage(1);
-	//Execute(ListUploadedFilesOperation, TArray<FLocalizationServiceTranslationIdentifier>(), InConcurrency);
-
-	// UPLOAD FILE
-	//TSharedRef<FOneSkyUploadFileOperation, ESPMode::ThreadSafe> UploadFileOperation = ILocalizationServiceOperation::Create<FOneSkyUploadFileOperation>();
-	//UploadFileOperation->SetInProjectId(27922);
-	//UploadFileOperation->SetInLocalFilePathAndName("D:\\zh_CN.po");
-	//UploadFileOperation->SetInOneSkyTargetFileName("en.po");
-	//UploadFileOperation->SetInLocale("zh_CN");
-	//UploadFileOperation->SetInIsKeepingAllStrings(false);
-	//Execute(UploadFileOperation, TArray<FLocalizationServiceTranslationIdentifier>(), InConcurrency);
-
-	// LIST UPLOADED FILES
-	//TSharedRef<FOneSkyListPhraseCollectionsOperation, ESPMode::ThreadSafe> ListPhraseCollectionsOperation = ILocalizationServiceOperation::Create<FOneSkyListPhraseCollectionsOperation>();
-	//ListPhraseCollectionsOperation->SetInProjectId(27922);
-	//ListPhraseCollectionsOperation->SetInItemsPerPage(1);
-	//ListPhraseCollectionsOperation->SetInPage(1);
-	//Execute(ListPhraseCollectionsOperation, TArray<FLocalizationServiceTranslationIdentifier>(), InConcurrency);
-
-	return ELocalizationServiceOperationCommandResult::Succeeded;
-}
-
 FText FOneSkyLocalizationServiceProvider::GetStatusText() const
 {
 	FOneSkyLocalizationServiceModule& OneSkyLocalizationService = FModuleManager::LoadModuleChecked<FOneSkyLocalizationServiceModule>( "OneSkyLocalizationService" );
@@ -458,44 +355,30 @@ void FOneSkyLocalizationServiceProvider::CustomizeSettingsDetails(IDetailCategor
 		];
 }
 
-static void ProjectGroupChanged(const FText& NewText, ETextCommit::Type CommitType, FGuid TargetGuid)
-{
-	FOneSkyLocalizationTargetSetting* Settings = FOneSkyLocalizationServiceModule::Get().AccessSettings().GetSettingsForTarget(TargetGuid);
-	FOneSkyLocalizationServiceModule::Get().AccessSettings().SetSettingsForTarget(TargetGuid, NewText.ToString(), Settings->OneSkyProjectName, Settings->OneSkyFileName);
-}
-
 static void ProjectChanged(const FText& NewText, ETextCommit::Type CommitType, FGuid TargetGuid)
 {
-	FOneSkyLocalizationTargetSetting* Settings = FOneSkyLocalizationServiceModule::Get().AccessSettings().GetSettingsForTarget(TargetGuid);
-	FOneSkyLocalizationServiceModule::Get().AccessSettings().SetSettingsForTarget(TargetGuid, Settings->OneSkyProjectGroupName, NewText.ToString(), Settings->OneSkyFileName);
+	FOneSkyLocalizationTargetSetting* Settings = FOneSkyLocalizationServiceModule::Get().AccessSettings().GetSettingsForTarget(TargetGuid, true);
+	int32 NewProjectId = INDEX_NONE;	// Default to -1
+	FString StringId = NewText.ToString();
+	// Don't allow this to be set to a non-numeric value.
+	if (StringId.IsNumeric())
+	{
+		NewProjectId = FCString::Atoi(*StringId);
+	}
+	FOneSkyLocalizationServiceModule::Get().AccessSettings().SetSettingsForTarget(TargetGuid, NewProjectId, Settings->OneSkyFileName);
 }
 
 static void FileNameChanged(const FText& NewText, ETextCommit::Type CommitType, FGuid TargetGuid)
 {
-	FOneSkyLocalizationTargetSetting* Settings = FOneSkyLocalizationServiceModule::Get().AccessSettings().GetSettingsForTarget(TargetGuid);
-	FOneSkyLocalizationServiceModule::Get().AccessSettings().SetSettingsForTarget(TargetGuid, Settings->OneSkyProjectGroupName, Settings->OneSkyProjectName, NewText.ToString());
+	FOneSkyLocalizationTargetSetting* Settings = FOneSkyLocalizationServiceModule::Get().AccessSettings().GetSettingsForTarget(TargetGuid, true);
+	FOneSkyLocalizationServiceModule::Get().AccessSettings().SetSettingsForTarget(TargetGuid, Settings->OneSkyProjectId, NewText.ToString());
 }
 
 void FOneSkyLocalizationServiceProvider::CustomizeTargetDetails(IDetailCategoryBuilder& DetailCategoryBuilder, const FGuid& TargetGuid) const
 {
-	//Target.Settings.Guid.ToString();
-	FOneSkyLocalizationTargetSetting* Settings = FOneSkyLocalizationServiceModule::Get().AccessSettings().GetSettingsForTarget(TargetGuid);
-
-	FText ProjectGroupText = LOCTEXT("OneSkyProjectGroupLabel", "OneSky Project Group");
-	FDetailWidgetRow& ProjectGroupRow = DetailCategoryBuilder.AddCustomRow(ProjectGroupText);
-	ProjectGroupRow.NameContent()
-		[
-			SNew(STextBlock)
-			.Text(ProjectGroupText)
-		];
-	ProjectGroupRow.ValueContent()
-		[
-			SNew(SEditableTextBox)
-			.OnTextCommitted(FOnTextCommitted::CreateStatic(&ProjectGroupChanged, TargetGuid))
-			.Text(FText::FromString(Settings->OneSkyProjectGroupName))
-		];
+	FOneSkyLocalizationTargetSetting* Settings = FOneSkyLocalizationServiceModule::Get().AccessSettings().GetSettingsForTarget(TargetGuid, true);
 
-	FText ProjectText = LOCTEXT("OneSkyProjectLabel", "OneSky Project");
+	FText ProjectText = LOCTEXT("OneSkyProjectIdLabel", "OneSky Project ID");
 	FDetailWidgetRow& ProjectRow = DetailCategoryBuilder.AddCustomRow(ProjectText);
 	ProjectRow.NameContent()
 		[
@@ -506,7 +389,18 @@ void FOneSkyLocalizationServiceProvider::CustomizeTargetDetails(IDetailCategoryB
 		[
 			SNew(SEditableTextBox)
 			.OnTextCommitted(FOnTextCommitted::CreateStatic(&ProjectChanged, TargetGuid))
-			.Text(FText::FromString(Settings->OneSkyProjectName))
+			.Text_Lambda([Settings]
+			{
+				int32 SavedProjectId = Settings->OneSkyProjectId;
+				if (SavedProjectId >= 0)
+				{
+					return (const FText&) FText::FromString(FString::FromInt(SavedProjectId));
+				}
+				
+				// Show empty string if value is default (-1)
+				return FText::GetEmpty();
+			}
+			)
 		];
 
 	FText FileText = LOCTEXT("OneSkyFileNameLabel", "OneSky File Name");
diff --git a/Engine/Plugins/Developer/OneSkyLocalizationService/Source/OneSkyLocalizationService/Private/OneSkyLocalizationServiceProvider.h b/Engine/Plugins/Developer/OneSkyLocalizationService/Source/OneSkyLocalizationService/Private/OneSkyLocalizationServiceProvider.h
--- a/Engine/Plugins/Developer/OneSkyLocalizationService/Source/OneSkyLocalizationService/Private/OneSkyLocalizationServiceProvider.h
+++ b/Engine/Plugins/Developer/OneSkyLocalizationService/Source/OneSkyLocalizationService/Private/OneSkyLocalizationServiceProvider.h
@@ -83,9 +83,6 @@ public:
 	/** Helper function used to update state cache */
 	TSharedRef<FOneSkyLocalizationServiceState, ESPMode::ThreadSafe> GetStateInternal(const FLocalizationServiceTranslationIdentifier& InTranslationId);
 
-	/** Helper function used to update project groups cache */
-	bool RetrieveProjectGroups(ELocalizationServiceOperationConcurrency::Type InConcurrency);
-
 	/**
 	 * Connects to the Localization service server if the persistent connection is not already established.
 	 *
@@ -128,12 +125,6 @@ private:
 
 	/** A pointer to the persistent P4 connection for synchronous operations */
 	class FOneSkyConnection* PersistentConnection;
-
-	/** Project Group Cache (group name to group id) */
-	TMap<FString, int32> ProjectGroupCache;
-
-	/** Project Cache (project name to project id) per group id */
-	TMap<int32, TMap<FString, int32>> ProjectCachePerGroupId;
 	
 	/** State cache */
 	TMap<FLocalizationServiceTranslationIdentifier, TSharedRef<class FOneSkyLocalizationServiceState, ESPMode::ThreadSafe>, FDefaultSetAllocator, FLocalizationServiceTranslationIdentifierKeyFuncs<TSharedRef<class FOneSkyLocalizationServiceState, ESPMode::ThreadSafe>> >  StateCache;
diff --git a/Engine/Plugins/Developer/OneSkyLocalizationService/Source/OneSkyLocalizationService/Private/OneSkyLocalizationServiceSettings.cpp b/Engine/Plugins/Developer/OneSkyLocalizationService/Source/OneSkyLocalizationService/Private/OneSkyLocalizationServiceSettings.cpp
--- a/Engine/Plugins/Developer/OneSkyLocalizationService/Source/OneSkyLocalizationService/Private/OneSkyLocalizationServiceSettings.cpp
+++ b/Engine/Plugins/Developer/OneSkyLocalizationService/Source/OneSkyLocalizationService/Private/OneSkyLocalizationServiceSettings.cpp
@@ -158,14 +158,14 @@ FOneSkyConnectionInfo FOneSkyLocalizationServiceSettings::GetConnectionInfo() co
 	return OutConnectionInfo;
 }
 
-FOneSkyLocalizationTargetSetting* FOneSkyLocalizationServiceSettings::GetSettingsForTarget(FGuid TargetGuid)
+FOneSkyLocalizationTargetSetting* FOneSkyLocalizationServiceSettings::GetSettingsForTarget(FGuid TargetGuid, bool bCreateIfNotFound /* = false */)
 {
 	FOneSkyLocalizationTargetSetting* Settings = TargetSettingsObject->TargetSettings.FindByPredicate([&](const FOneSkyLocalizationTargetSetting& Setting)
 	{
 		return Setting.TargetGuid == TargetGuid;
 	});
 
-	if (!Settings)
+	if (!Settings && bCreateIfNotFound)
 	{
 		FOneSkyLocalizationTargetSetting NewSettings;
 		NewSettings.TargetGuid = TargetGuid;
@@ -176,7 +176,7 @@ FOneSkyLocalizationTargetSetting* FOneSkyLocalizationServiceSettings::GetSetting
 	return Settings;
 }
 
-void FOneSkyLocalizationServiceSettings::SetSettingsForTarget(FGuid TargetGuid, FString ProjectGroupName, FString ProjectName, FString FileName)
+void FOneSkyLocalizationServiceSettings::SetSettingsForTarget(FGuid TargetGuid, int32 ProjectId, FString FileName)
 {
 	FOneSkyLocalizationTargetSetting* Settings = TargetSettingsObject->TargetSettings.FindByPredicate([&](const FOneSkyLocalizationTargetSetting& Setting)
 	{
@@ -191,8 +191,7 @@ void FOneSkyLocalizationServiceSettings::SetSettingsForTarget(FGuid TargetGuid,
 		Settings = &(TargetSettingsObject->TargetSettings[NewIndex]);
 	}
 
-	Settings->OneSkyProjectGroupName = ProjectGroupName;
-	Settings->OneSkyProjectName = ProjectName;
+	Settings->OneSkyProjectId = ProjectId;
 	Settings->OneSkyFileName = FileName;
 
 	TargetSettingsObject->SaveConfig();
diff --git a/Engine/Plugins/Developer/OneSkyLocalizationService/Source/OneSkyLocalizationService/Private/OneSkyLocalizationServiceSettings.h b/Engine/Plugins/Developer/OneSkyLocalizationService/Source/OneSkyLocalizationService/Private/OneSkyLocalizationServiceSettings.h
--- a/Engine/Plugins/Developer/OneSkyLocalizationService/Source/OneSkyLocalizationService/Private/OneSkyLocalizationServiceSettings.h
+++ b/Engine/Plugins/Developer/OneSkyLocalizationService/Source/OneSkyLocalizationService/Private/OneSkyLocalizationServiceSettings.h
@@ -13,6 +13,8 @@ struct FOneSkyLocalizationTargetSetting
 {
 	GENERATED_USTRUCT_BODY()
 
+		FOneSkyLocalizationTargetSetting() : OneSkyProjectId(INDEX_NONE) {}
+
 	/**
 	* The GUID of the LocalizationTarget that these OneSky settings are for
 	*/
@@ -20,16 +22,10 @@ struct FOneSkyLocalizationTargetSetting
 		FGuid TargetGuid;
 
 	/**
-	* The name of the OneSky Project Group this target belongs to
-	*/
-	UPROPERTY()
-		FString OneSkyProjectGroupName;
-
-	/**
-	* The name of the OneSky Project  this target belongs to
+	* The id of the OneSky Project  this target belongs to
 	*/
 	UPROPERTY()
-		FString OneSkyProjectName;
+		int32 OneSkyProjectId;
 
 	/**
 	* The name of the file that corresponds to this target
@@ -100,9 +96,14 @@ public:
 		return bSaveSecretKey;
 	}
 
-	FOneSkyLocalizationTargetSetting* GetSettingsForTarget(FGuid TargetGuid);
+	/** Get settings for a specific localization target by it's Guid
+	*	@param TargetGuid - the GUID to search for
+	*	@param bCreateIfNotFound - Add a new blank settings entry to the array if not found
+	*/
+	FOneSkyLocalizationTargetSetting* GetSettingsForTarget(FGuid TargetGuid, bool bCreateIfNotFound = false);
 
-	void SetSettingsForTarget(FGuid TargetGuid, FString ProjectGroupName, FString ProjectName, FString FileName);
+	/** Set settings for a specific localization target by it's Guid */
+	void SetSettingsForTarget(FGuid TargetGuid, int32 ProjectId, FString FileName);
 
 public:
 	/** Object that contains array of OneSky settings for a given targets, used to serialize to .ini file */
diff --git a/Engine/Source/Developer/LocalizationService/Public/LocalizationServiceOperations.h b/Engine/Source/Developer/LocalizationService/Public/LocalizationServiceOperations.h
--- a/Engine/Source/Developer/LocalizationService/Public/LocalizationServiceOperations.h
+++ b/Engine/Source/Developer/LocalizationService/Public/LocalizationServiceOperations.h
@@ -51,6 +51,67 @@ protected:
 	FText OutErrorText;
 };
 
+/**
+ * Operation used to commit localization(s) to localization service
+ */
+class FDownloadLocalizationTargetFile : public ILocalizationServiceOperation
+{
+public:
+
+	// ILocalizationServiceOperation interface
+	virtual FName GetName() const override
+	{
+		return "DownloadLocalizationTargetFile";
+	}
+
+	virtual FText GetInProgressString() const override
+	{ 
+		return LOCTEXT("LocalizationService_FLocalizationServicerDownloadLocalizationTargetFile", "Getting Localization Target File from localization Service..."); 
+	}
+
+
+	void SetInTargetGuid(const FGuid& NewInTargetGuid)
+	{
+		InTargetGuid = NewInTargetGuid;
+	}
+
+	const FGuid& GetInTargetGuid() const
+	{
+		return InTargetGuid;
+	}
+
+	void SetInRelativeOutputFilePathAndName(const FString& NewInRelativeOutputFilePathAndName)
+	{
+		InRelativeOutputFilePathAndName = NewInRelativeOutputFilePathAndName;
+	}
+
+	const FString& GetInRelativeOutputFilePathAndName() const
+	{
+		return InRelativeOutputFilePathAndName;
+	}
+
+	FString GetInLocale()
+	{
+		return InLocale;
+	}
+
+	void SetInLocale(FString NewLocale)
+	{
+		InLocale = NewLocale;
+	}
+
+protected:
+
+	/** The GUID of the Localization Target */
+	FGuid InTargetGuid;
+
+	/** The locale (culture code, for example "fr" or "en-us" or "ja-jp") to download the translations for */
+	FString InLocale;
+
+	/** The path and name to the file downloaded, relative to project directory */
+	FString InRelativeOutputFilePathAndName;
+};
+
 ///**
 // * Operation used to commit localization(s) to localization service
 // */
