OneSky Plugin: If editor is closed between submitting the HTTP request and receiving the callback, the callback is never called and the thread never finishes, so abandon the threads if we are in this case.

diff --git a/Engine/Plugins/Developer/OneSkyLocalizationService/Source/OneSkyLocalizationService/Private/OneSkyLocalizationServiceCommand.cpp b/Engine/Plugins/Developer/OneSkyLocalizationService/Source/OneSkyLocalizationService/Private/OneSkyLocalizationServiceCommand.cpp
--- a/Engine/Plugins/Developer/OneSkyLocalizationService/Source/OneSkyLocalizationService/Private/OneSkyLocalizationServiceCommand.cpp
+++ b/Engine/Plugins/Developer/OneSkyLocalizationService/Source/OneSkyLocalizationService/Private/OneSkyLocalizationServiceCommand.cpp
@@ -32,6 +32,12 @@ bool FOneSkyLocalizationServiceCommand::DoWork()
 	while (!bExecuteProcessed)
 	{
 		FPlatformProcess::Sleep(0.01f);
+
+		// If the editor was closed, our callbacks won't get called anyway, so just kill the thread
+		if (GIsRequestingExit)
+		{
+			Abandon();
+		}
 	}
 
 	return bCommandSuccessful;
