Translation Picker: Upload changes to OneSky

diff --git a/Engine/Plugins/Developer/OneSkyLocalizationService/Source/OneSkyLocalizationService/Private/OneSkyLocalizationServiceModule.cpp b/Engine/Plugins/Developer/OneSkyLocalizationService/Source/OneSkyLocalizationService/Private/OneSkyLocalizationServiceModule.cpp
--- a/Engine/Plugins/Developer/OneSkyLocalizationService/Source/OneSkyLocalizationService/Private/OneSkyLocalizationServiceModule.cpp
+++ b/Engine/Plugins/Developer/OneSkyLocalizationService/Source/OneSkyLocalizationService/Private/OneSkyLocalizationServiceModule.cpp
@@ -39,7 +39,10 @@ void FOneSkyLocalizationServiceModule::StartupModule()
 
 	// Files API
 	OneSkyLocalizationServiceProvider.RegisterWorker("ListUploadedFiles", FGetOneSkyLocalizationServiceWorker::CreateStatic(&CreateWorker<FOneSkyListUploadedFilesWorker>));
-	OneSkyLocalizationServiceProvider.RegisterWorker("UploadFile", FGetOneSkyLocalizationServiceWorker::CreateStatic(&CreateWorker<FOneSkyUploadFileWorker>));
+	OneSkyLocalizationServiceProvider.RegisterWorker("UploadLocalizationTargetFile", FGetOneSkyLocalizationServiceWorker::CreateStatic(&CreateWorker<FOneSkyUploadFileWorker>));
+
+	// Import API
+	OneSkyLocalizationServiceProvider.RegisterWorker("ShowImportTask", FGetOneSkyLocalizationServiceWorker::CreateStatic(&CreateWorker<FOneSkyShowImportTaskWorker>));
 
 	// Phrase Collections API
 	OneSkyLocalizationServiceProvider.RegisterWorker("ListPhraseCollections", FGetOneSkyLocalizationServiceWorker::CreateStatic(&CreateWorker<FOneSkyListPhraseCollectionsWorker>));
diff --git a/Engine/Plugins/Developer/OneSkyLocalizationService/Source/OneSkyLocalizationService/Private/OneSkyLocalizationServiceOperations.cpp b/Engine/Plugins/Developer/OneSkyLocalizationService/Source/OneSkyLocalizationService/Private/OneSkyLocalizationServiceOperations.cpp
--- a/Engine/Plugins/Developer/OneSkyLocalizationService/Source/OneSkyLocalizationService/Private/OneSkyLocalizationServiceOperations.cpp
+++ b/Engine/Plugins/Developer/OneSkyLocalizationService/Source/OneSkyLocalizationService/Private/OneSkyLocalizationServiceOperations.cpp
@@ -12,6 +12,7 @@
 #include "JsonStructDeserializerBackend.h"
 #include "StructDeserializer.h"
 #include "OneSkyLocalizationServiceResponseTypes.h"
+#include "MessageLog.h"
 
 #define LOCTEXT_NAMESPACE "OneSkyLocalizationService"
 
@@ -925,25 +926,35 @@ bool FOneSkyUploadFileWorker::Execute(class FOneSkyLocalizationServiceCommand& I
 {
 	// Store pointer to command so we can access it in the http request callback
 	Command = &InCommand;
-	TSharedPtr<FOneSkyUploadFileOperation, ESPMode::ThreadSafe> UploadFilesOp = StaticCastSharedRef<FOneSkyUploadFileOperation>(InCommand.Operation);
+	TSharedPtr<FUploadLocalizationTargetFile, ESPMode::ThreadSafe> UploadFileOp = StaticCastSharedRef<FUploadLocalizationTargetFile>(InCommand.Operation);
 
-	int32 InProjectId = -1;
+	FGuid InTargetGuid;
 	FString InFilePathAndName = "";
-	FString InOneSkyTargetFileName = "";
 	FString InLocale = "";
-	bool bInIsKeepingAllStrings = false;
+	bool bInIsKeepingAllStrings = true;
 	FString InFileFormat = "";
 
-	if (UploadFilesOp.IsValid())
+	if (UploadFileOp.IsValid())
 	{
-		InProjectId = UploadFilesOp->GetInProjectId();
-		InFilePathAndName = UploadFilesOp->GetInLocalFilePathAndName();
-		InOneSkyTargetFileName = UploadFilesOp->GetInOneSkyTargetFileName();
-		InLocale = UploadFilesOp->GetInLocale();
-		bInIsKeepingAllStrings = UploadFilesOp->GetInIsKeepingAllStrings();
+		InTargetGuid = UploadFileOp->GetInTargetGuid();
+		InLocale = UploadFileOp->GetInLocale();
+		// Path is relative to game directory
+		InFilePathAndName = FPaths::ConvertRelativePathToFull(FPaths::GameDir() / UploadFileOp->GetInRelativeInputFilePathAndName());
+		bInIsKeepingAllStrings = UploadFileOp->GetPreserveAllText();
 		InFileFormat = GetFileFormat(FPaths::GetExtension(InFilePathAndName, true));
 	}
 
+	int32 InProjectId = -1;
+	FString InOneSkyTargetFileName = "";
+
+	// OneSky project settings are accessed by their localization target guid. These settings are set in the Localization Dashboard
+	FOneSkyLocalizationTargetSetting* Settings = FOneSkyLocalizationServiceModule::Get().AccessSettings().GetSettingsForTarget(InTargetGuid);
+	if (Settings != nullptr)
+	{
+		InProjectId = Settings->OneSkyProjectId;
+		InOneSkyTargetFileName = FGenericPlatformHttp::UrlEncode(Settings->OneSkyFileName);
+	}
+
 	FString Url = AddAuthenticationParameters(InCommand.ConnectionInfo, "https://platform.api.onesky.io/1/projects/" + FString::FromInt(InProjectId) / "files");
 	Url += "&file_format=" + InFileFormat + "&locale=" + InLocale;
 	if (!bInIsKeepingAllStrings)
@@ -992,6 +1003,34 @@ void FOneSkyUploadFileWorker::Query_HttpRequestComplete(FHttpRequestPtr HttpRequ
 		bResult = DeserializeResponseToStruct(&OutUploadFileResponse, StructType, HttpResponse);
 	}
 
+	if (bResult)
+	{
+		TSharedPtr<FUploadLocalizationTargetFile, ESPMode::ThreadSafe> UploadFileOp = StaticCastSharedRef<FUploadLocalizationTargetFile>(Command->Operation);
+
+		FGuid InTargetGuid;
+		int32 InProjectId = -1;
+
+		if (UploadFileOp.IsValid())
+		{
+			InTargetGuid = UploadFileOp->GetInTargetGuid();
+
+			// OneSky project settings are accessed by their localization target guid. These settings are set in the Localization Dashboard
+			FOneSkyLocalizationTargetSetting* Settings = FOneSkyLocalizationServiceModule::Get().AccessSettings().GetSettingsForTarget(InTargetGuid);
+			if (Settings != nullptr)
+			{
+				InProjectId = Settings->OneSkyProjectId;
+			}
+		}
+
+		FShowImportTaskQueueItem ImportTaskQueueItem;
+		ImportTaskQueueItem.ImportId = OutUploadFileResponse.data.import.id;
+		ImportTaskQueueItem.ProjectId = InProjectId;
+		// Wait one minute before querying the status of the import
+		ImportTaskQueueItem.ExecutionDelayInSeconds = 60.0f;
+		ImportTaskQueueItem.CreatedTimestamp = FDateTime::UtcNow();
+		FOneSkyLocalizationServiceModule::Get().GetProvider().EnqueShowImportTask(ImportTaskQueueItem);
+	}
+
 	if (Command != nullptr)
 	{
 		Command->bCommandSuccessful = bResult;
@@ -1069,6 +1108,114 @@ void FOneSkyListPhraseCollectionsWorker::Query_HttpRequestComplete(FHttpRequestP
 // end LIST PHRASE COLLECTIONS
 
 
+// SHOW IMPORT TASK
+
+FName FOneSkyShowImportTaskWorker::GetName() const
+{
+	return "ShowImportTask";
+}
+
+bool FOneSkyShowImportTaskWorker::Execute(class FOneSkyLocalizationServiceCommand& InCommand)
+{
+	// Store pointer to command so we can access it in the http request callback
+	Command = &InCommand;
+
+	TSharedPtr<FOneSkyShowImportTaskOperation, ESPMode::ThreadSafe> ShowImportTaskOp = StaticCastSharedRef<FOneSkyShowImportTaskOperation>(InCommand.Operation);
+
+	int32 InProjectId = -1;
+	int32 InImportId = -1;
+	int32 InExecutionDelayInSeconds = -1;
+	FDateTime InCreationTimeStamp = FDateTime::FDateTime(1970, 1, 1);
+
+	if (ShowImportTaskOp.IsValid())
+	{
+		InProjectId = ShowImportTaskOp->GetInProjectId();
+		InImportId = ShowImportTaskOp->GetInImportId();
+		InExecutionDelayInSeconds = ShowImportTaskOp->GetInExecutionDelayInSeconds();
+		InCreationTimeStamp = ShowImportTaskOp->GetInCreationTimestamp();
+	}
+
+	double ElapsedTimeInSeconds = (FDateTime::UtcNow() - InCreationTimeStamp).GetTotalSeconds();
+
+	while (!GIsRequestingExit && (ElapsedTimeInSeconds <= InExecutionDelayInSeconds))
+	{
+		ElapsedTimeInSeconds = (FDateTime::UtcNow() - InCreationTimeStamp).GetTotalSeconds();
+		FPlatformProcess::Sleep(0.05);
+	}
+
+	FString Url = "https://platform.api.onesky.io/1/projects/" + FString::FromInt(InProjectId) / "import-tasks/" + FString::FromInt(InImportId);
+	Url = AddAuthenticationParameters(InCommand.ConnectionInfo, Url);
+
+	FHttpModule& HttpModule = FModuleManager::LoadModuleChecked<FHttpModule>("HTTP");
+	TSharedRef<class IHttpRequest> HttpRequest = HttpModule.Get().CreateRequest();
+
+	// kick off http request to read
+	HttpRequest->OnProcessRequestComplete().BindRaw(this, &FOneSkyShowImportTaskWorker::Query_HttpRequestComplete);
+	HttpRequest->SetURL(Url);
+	HttpRequest->SetHeader(TEXT("Content-Type"), TEXT("application/json; charset=utf-8"));
+	HttpRequest->SetVerb(TEXT("GET"));
+	HttpRequest->ProcessRequest();
+
+	return InCommand.bCommandSuccessful;
+
+	return false;
+}
+
+void FOneSkyShowImportTaskWorker::Query_HttpRequestComplete(FHttpRequestPtr HttpRequest, FHttpResponsePtr HttpResponse, bool bSucceeded)
+{
+	bool bResult = false;
+
+	if (bSucceeded)
+	{
+		UScriptStruct& StructType = *(FOneSkyShowImportTaskResponse::StaticStruct());
+		bResult = DeserializeResponseToStruct(&OutShowImportTaskResponse, StructType, HttpResponse);
+	}
+
+	if (bResult)
+	{
+		TSharedPtr<FOneSkyShowImportTaskOperation, ESPMode::ThreadSafe> ShowImportTaskOp = StaticCastSharedRef<FOneSkyShowImportTaskOperation>(Command->Operation);
+
+		int32 InProjectId = -1;
+		int32 InImportId = -1;
+		int32 InExecutionDelayInSeconds = -1;
+
+		if (ShowImportTaskOp.IsValid())
+		{
+			InProjectId = ShowImportTaskOp->GetInProjectId();
+			InImportId = ShowImportTaskOp->GetInImportId();
+			InExecutionDelayInSeconds = ShowImportTaskOp->GetInExecutionDelayInSeconds();
+
+			if (OutShowImportTaskResponse.data.status == "in-progress")
+			{
+				FShowImportTaskQueueItem ImportTaskQueueItem;
+				ImportTaskQueueItem.ImportId = InImportId;
+				ImportTaskQueueItem.ProjectId = InProjectId;
+				ImportTaskQueueItem.ExecutionDelayInSeconds = InExecutionDelayInSeconds;
+				ImportTaskQueueItem.CreatedTimestamp = FDateTime::UtcNow();
+				FOneSkyLocalizationServiceModule::Get().GetProvider().EnqueShowImportTask(ImportTaskQueueItem);
+			}
+			else if (OutShowImportTaskResponse.data.status == "failed")
+			{
+				FString CultureName = OutShowImportTaskResponse.data.file.locale.code;
+				FString TargetName = OutShowImportTaskResponse.data.file.name;
+				FText FailureText = FText::Format(LOCTEXT("ImportTaskFailed", "{0} translations for {1} target failed to import to OneSky!"), FText::FromString(CultureName), FText::FromString(TargetName));
+				FMessageLog LocalizationServiceMessageLog("TranslationEditor");
+				LocalizationServiceMessageLog.Error(FailureText);
+				LocalizationServiceMessageLog.Notify(FailureText, EMessageSeverity::Error, true);
+			}
+		}
+	}
+
+	if (Command != nullptr)
+	{
+		Command->bCommandSuccessful = bResult;
+		FPlatformAtomics::InterlockedExchange(&(Command->bExecuteProcessed), 1);
+	}
+}
+
+// end SHOW IMPORT TASK
+
+
 //FName FOneSkyListProjectTypesWorker::GetName() const
 //{
 //	return FName(TEXT("FOneSkyListProjectTypesWorker"));
diff --git a/Engine/Plugins/Developer/OneSkyLocalizationService/Source/OneSkyLocalizationService/Private/OneSkyLocalizationServiceOperations.h b/Engine/Plugins/Developer/OneSkyLocalizationService/Source/OneSkyLocalizationService/Private/OneSkyLocalizationServiceOperations.h
--- a/Engine/Plugins/Developer/OneSkyLocalizationService/Source/OneSkyLocalizationService/Private/OneSkyLocalizationServiceOperations.h
+++ b/Engine/Plugins/Developer/OneSkyLocalizationService/Source/OneSkyLocalizationService/Private/OneSkyLocalizationServiceOperations.h
@@ -719,66 +719,67 @@
 // end LIST UPLOADED FILES
 
 
-// UPLOADED FILE
+//// UPLOAD FILE
 
+class FOneSkyUploadFileWorker : public IOneSkyLocalizationServiceWorker
+{
+public:
+	virtual ~FOneSkyUploadFileWorker() {}
+	// IOneSkyLocalizationServiceWorker interface
+	virtual FName GetName() const override;
+	virtual bool Execute(class FOneSkyLocalizationServiceCommand& InCommand) override;
+	virtual bool UpdateStates() const override {/*TODO:*/return true; }
+	void Query_HttpRequestComplete(FHttpRequestPtr HttpRequest, FHttpResponsePtr HttpResponse, bool bSucceeded);
+public:
+
+	FOneSkyLocalizationServiceCommand* Command;
+
+	/** Response to the  command */
+	FOneSkyUploadFileResponse OutUploadFileResponse;
+};
+
+// end UPLOAD FILE
+
+
+// LIST PHRASE COLLECTIONS
+
 /**
-* Operation used to upload a files to a project in one sky
+* Operation used to list phrase collections for a project in one sky
 */
-class FOneSkyUploadFileOperation : public ILocalizationServiceOperation
+class FOneSkyListPhraseCollectionsOperation : public ILocalizationServiceOperation
 {
 public:
-
 	// ILocalizationServiceOperation interface
 	virtual FName GetName() const override
 	{
-		return "UploadFile";
+		return "ListPhraseCollections";
 	}
 
 	virtual FText GetInProgressString() const override
 	{
-		return LOCTEXT("OneSkyLocalizationService_UploadFileWorker", "Uploading File to a Project in OneSky...");
+		return LOCTEXT("OneSkyLocalizationService_ListPhraseCollections", "Listing Phrase Collections for a Project in OneSky...");
 	}
 
-	FOneSkyUploadFileOperation() : InProjectId(-1), bInIsKeepingAllStrings(true) { }
+	FOneSkyListPhraseCollectionsOperation() : InProjectId(-1), InPage(-1), InItemsPerPage(-1) { }
 
-	void SetInLocalFilePathAndName(FString NewInLocalFilePathAndName)
+	void SetInPage(int32 NewPage)
 	{
-		InLocalFilePathAndName = NewInLocalFilePathAndName;
+		InPage = NewPage;
 	}
 
-	FString GetInLocalFilePathAndName()
+	int32 GetInPage()
 	{
-		return InLocalFilePathAndName;
+		return InPage;
 	}
 
-	void SetInOneSkyTargetFileName(FString NewInOneSkyTargetFileName)
+	void SetInItemsPerPage(int32 NewItemsPerPage)
 	{
-		InOneSkyTargetFileName = NewInOneSkyTargetFileName;
+		InItemsPerPage = NewItemsPerPage;
 	}
 
-	FString GetInOneSkyTargetFileName()
+	int32 GetInItemsPerPage()
 	{
-		return InOneSkyTargetFileName;
+		return InItemsPerPage;
 	}
 
-	void SetInLocale(FString NewInLocale)
-	{
-		InLocale = NewInLocale;
-	}
-
-	FString GetInLocale()
-	{
-		return InLocale;
-	}
-
-	bool GetInIsKeepingAllStrings()
-	{
-		return bInIsKeepingAllStrings;
-	}
-
-	void SetInIsKeepingAllStrings(bool NewInIsKeepingAllStrings)
-	{
-		bInIsKeepingAllStrings = NewInIsKeepingAllStrings;
-	}
-
 	int32 GetInProjectId()
@@ -794,16 +795,14 @@
 protected:
 
 	int32 InProjectId;
-	FString InLocalFilePathAndName;
-	FString InOneSkyTargetFileName;
-	FString InLocale;
-	bool bInIsKeepingAllStrings;
+	int32 InPage;
+	int32 InItemsPerPage;
 };
 
-class FOneSkyUploadFileWorker : public IOneSkyLocalizationServiceWorker
+class FOneSkyListPhraseCollectionsWorker : public IOneSkyLocalizationServiceWorker
 {
 public:
-	virtual ~FOneSkyUploadFileWorker() {}
+	virtual ~FOneSkyListPhraseCollectionsWorker() {}
 	// IOneSkyLocalizationServiceWorker interface
 	virtual FName GetName() const override;
 	virtual bool Execute(class FOneSkyLocalizationServiceCommand& InCommand) override;
@@ -814,74 +813,85 @@
 	FOneSkyLocalizationServiceCommand* Command;
 
 	/** Response to the  command */
-	FOneSkyUploadFileResponse OutUploadFileResponse;
+	FOneSkyListPhraseCollectionsResponse OutListPhraseCollectionsResponse;
 };
 
-// end UPLOAD FILES
+// end LIST PHRASE COLLECTIONS
 
 
-// LIST PHRASE COLLECTIONS
+// SHOW IMPORT TASK
 
 /**
-* Operation used to list phrase collections for a project in one sky
+* Operation used to show a project in one sky
 */
-class FOneSkyListPhraseCollectionsOperation : public ILocalizationServiceOperation
+class FOneSkyShowImportTaskOperation : public ILocalizationServiceOperation
 {
 public:
 	// ILocalizationServiceOperation interface
 	virtual FName GetName() const override
 	{
-		return "ListPhraseCollections";
+		return "ShowImportTask";
 	}
 
 	virtual FText GetInProgressString() const override
 	{
-		return LOCTEXT("OneSkyLocalizationService_ListPhraseCollections", "Listing Phrase Collections for a Project in OneSky...");
+		return LOCTEXT("OneSkyLocalizationService_ShowImportTask", "Showing Import Task in OneSky...");
 	}
 
-	FOneSkyListPhraseCollectionsOperation() : InProjectId(-1), InPage(-1), InItemsPerPage(-1) { }
+	FOneSkyShowImportTaskOperation() : InProjectId(-1) { }
 
-	void SetInPage(int32 NewPage)
-	{
-		InPage = NewPage;
-	}
-
-	int32 GetInPage()
-	{
-		return InPage;
-	}
-
-	void SetInItemsPerPage(int32 NewItemsPerPage)
-	{
-		InItemsPerPage = NewItemsPerPage;
-	}
-
-	int32 GetInItemsPerPage()
-	{
-		return InItemsPerPage;
-	}
-
 	int32 GetInProjectId()
 	{
 		return InProjectId;
 	}
 
-	void SetInProjectId(int32 NewProjectGroupId)
+	void SetInProjectId(int32 NewProjectId)
 	{
-		InProjectId = NewProjectGroupId;
+		InProjectId = NewProjectId;
 	}
 
+	int32 GetInImportId()
+	{
+		return InImportId;
+	}
+
+	void SetInImportId(int32 NewImportId)
+	{
+		InImportId = NewImportId;
+	}
+
+	int32 GetInExecutionDelayInSeconds()
+	{
+		return InExecutionDelayInSeconds;
+	}
+
+	void SetInExecutionDelayInSeconds(int32 NewInExecutionDelayInSeconds)
+	{
+		InExecutionDelayInSeconds = NewInExecutionDelayInSeconds;
+	}
+
+	FDateTime GetInCreationTimestamp()
+	{
+		return InCreationTimestamp;
+	}
+
+	void SetInCreationTimestamp(FDateTime NewInCreationTimestamp)
+	{
+		InCreationTimestamp = NewInCreationTimestamp;
+	}
+
 protected:
 
 	int32 InProjectId;
-	int32 InPage;
-	int32 InItemsPerPage;
+	int32 InImportId;
+	int32 InExecutionDelayInSeconds;
+	FDateTime InCreationTimestamp;
 };
 
-class FOneSkyListPhraseCollectionsWorker : public IOneSkyLocalizationServiceWorker
+class FOneSkyShowImportTaskWorker : public IOneSkyLocalizationServiceWorker
 {
 public:
-	virtual ~FOneSkyListPhraseCollectionsWorker() {}
+	virtual ~FOneSkyShowImportTaskWorker() {}
 	// IOneSkyLocalizationServiceWorker interface
 	virtual FName GetName() const override;
 	virtual bool Execute(class FOneSkyLocalizationServiceCommand& InCommand) override;
@@ -891,11 +901,11 @@
 
 	FOneSkyLocalizationServiceCommand* Command;
 
-	/** Response to the  command */
-	FOneSkyListPhraseCollectionsResponse OutListPhraseCollectionsResponse;
+	/** Response to the command */
+	FOneSkyShowImportTaskResponse OutShowImportTaskResponse;
 };
 
-// end LIST PHRASE COLLECTIONS
+// end SHOW IMPORT TASK
 
 //class FOneSkyListProjectTypesWorker : public IOneSkyLocalizationServiceWorker
 //{
diff --git a/Engine/Plugins/Developer/OneSkyLocalizationService/Source/OneSkyLocalizationService/Private/OneSkyLocalizationServiceProvider.cpp b/Engine/Plugins/Developer/OneSkyLocalizationService/Source/OneSkyLocalizationService/Private/OneSkyLocalizationServiceProvider.cpp
--- a/Engine/Plugins/Developer/OneSkyLocalizationService/Source/OneSkyLocalizationService/Private/OneSkyLocalizationServiceProvider.cpp
+++ b/Engine/Plugins/Developer/OneSkyLocalizationService/Source/OneSkyLocalizationService/Private/OneSkyLocalizationServiceProvider.cpp
@@ -242,6 +242,19 @@ void FOneSkyLocalizationServiceProvider::OutputCommandMessages(const FOneSkyLoca
 void FOneSkyLocalizationServiceProvider::Tick()
 {	
 	bool bStatesUpdated = false;
+
+	while (!ShowImportTaskQueue.IsEmpty())
+	{
+		FShowImportTaskQueueItem ImportTaskItem;
+		ShowImportTaskQueue.Dequeue(ImportTaskItem);
+		TSharedRef<FOneSkyShowImportTaskOperation, ESPMode::ThreadSafe> ShowImportTaskOperation = ILocalizationServiceOperation::Create<FOneSkyShowImportTaskOperation>();
+		ShowImportTaskOperation->SetInProjectId(ImportTaskItem.ProjectId);
+		ShowImportTaskOperation->SetInImportId(ImportTaskItem.ImportId);
+		ShowImportTaskOperation->SetInExecutionDelayInSeconds(ImportTaskItem.ExecutionDelayInSeconds);
+		ShowImportTaskOperation->SetInCreationTimestamp(ImportTaskItem.CreatedTimestamp);
+		FOneSkyLocalizationServiceModule::Get().GetProvider().Execute(ShowImportTaskOperation, TArray<FLocalizationServiceTranslationIdentifier>(), ELocalizationServiceOperationConcurrency::Asynchronous);
+	}
+
 	for (int32 CommandIndex = 0; CommandIndex < CommandQueue.Num(); ++CommandIndex)
 	{
 		FOneSkyLocalizationServiceCommand& Command = *CommandQueue[CommandIndex];
diff --git a/Engine/Plugins/Developer/OneSkyLocalizationService/Source/OneSkyLocalizationService/Private/OneSkyLocalizationServiceProvider.h b/Engine/Plugins/Developer/OneSkyLocalizationService/Source/OneSkyLocalizationService/Private/OneSkyLocalizationServiceProvider.h
--- a/Engine/Plugins/Developer/OneSkyLocalizationService/Source/OneSkyLocalizationService/Private/OneSkyLocalizationServiceProvider.h
+++ b/Engine/Plugins/Developer/OneSkyLocalizationService/Source/OneSkyLocalizationService/Private/OneSkyLocalizationServiceProvider.h
@@ -46,6 +46,15 @@ struct FLocalizationServiceTranslationIdentifierKeyFuncs : BaseKeyFuncs < TPair<
 	}
 };
 
+// Struct representing a show import task in our queue
+struct FShowImportTaskQueueItem
+{
+	int32 ProjectId;
+	int32 ImportId;
+	FDateTime CreatedTimestamp;
+	int32 ExecutionDelayInSeconds;
+};
+
 class FOneSkyLocalizationServiceProvider : public ILocalizationServiceProvider
 {
 public:
@@ -96,6 +105,12 @@ public:
 		return PersistentConnection;
 	}
 
+	// Enqueue a show import task for later exection
+	void EnqueShowImportTask(FShowImportTaskQueueItem QueueItem)
+	{
+		ShowImportTaskQueue.Enqueue(QueueItem);
+	}
+
 private:
 
 	/** Helper function used to create a worker for a particular operation */
@@ -134,4 +149,7 @@ private:
 
 	/** Queue for commands given by the main thread */
 	TArray < FOneSkyLocalizationServiceCommand* > CommandQueue;
-};
+
+	/** Queue for import status tasks  */
+	TQueue < FShowImportTaskQueueItem, EQueueMode::Mpsc > ShowImportTaskQueue;
+};
\ No newline at end of file
diff --git a/Engine/Plugins/Developer/OneSkyLocalizationService/Source/OneSkyLocalizationService/Private/OneSkyLocalizationServiceResponseTypes.h b/Engine/Plugins/Developer/OneSkyLocalizationService/Source/OneSkyLocalizationService/Private/OneSkyLocalizationServiceResponseTypes.h
--- a/Engine/Plugins/Developer/OneSkyLocalizationService/Source/OneSkyLocalizationService/Private/OneSkyLocalizationServiceResponseTypes.h
+++ b/Engine/Plugins/Developer/OneSkyLocalizationService/Source/OneSkyLocalizationService/Private/OneSkyLocalizationServiceResponseTypes.h
@@ -1204,4 +1204,151 @@ struct FOneSkyListPhraseCollectionsResponse
 	FOneSkyListPhraseCollectionsResponse(ENoInit) { }
 };
 
-// end LIST PHRASE COLLECTIONS
\ No newline at end of file
+// end LIST PHRASE COLLECTIONS
+
+
+// SHOW IMPORT TASK 
+
+USTRUCT()
+struct FOneSkyShowImportTaskResponseMeta
+{
+	GENERATED_USTRUCT_BODY()
+
+	UPROPERTY()
+	int32 status;
+
+	/** Default constructor. */
+	FOneSkyShowImportTaskResponseMeta()
+		: status(-1)
+	{ }
+
+	/** Creates an uninitialized instance. */
+	FOneSkyShowImportTaskResponseMeta(ENoInit) { }
+};
+
+USTRUCT()
+struct FOneSkyShowImportTaskResponseLocale
+{
+	GENERATED_USTRUCT_BODY()
+
+	UPROPERTY()
+	FString code;
+
+	UPROPERTY()
+		FString english_name;
+
+	UPROPERTY()
+		FString local_name;
+
+	UPROPERTY()
+		FString locale;
+
+	UPROPERTY()
+		FString region;
+
+	/** Default constructor. */
+	FOneSkyShowImportTaskResponseLocale()
+		: code("")
+		, english_name("")
+		, local_name("")
+		, locale("")
+		, region("")
+	{ }
+
+	/** Creates an uninitialized instance. */
+	FOneSkyShowImportTaskResponseLocale(ENoInit) { }
+};
+
+
+USTRUCT()
+struct FOneSkyShowImportTaskResponseFile
+{
+	GENERATED_USTRUCT_BODY()
+
+	UPROPERTY()
+	FString name;
+
+	UPROPERTY()
+	FString format;
+
+	UPROPERTY()
+	FOneSkyShowImportTaskResponseLocale locale;
+
+	/** Default constructor. */
+	FOneSkyShowImportTaskResponseFile()
+		: name("")
+		, format("")
+		, locale()
+	{ }
+
+	/** Creates an uninitialized instance. */
+	FOneSkyShowImportTaskResponseFile(ENoInit) { }
+};
+
+USTRUCT()
+struct FOneSkyShowImportTaskResponseData
+{
+	GENERATED_USTRUCT_BODY()
+
+	UPROPERTY()
+	int32 id;
+
+	UPROPERTY()
+	FOneSkyShowImportTaskResponseFile file;
+
+	UPROPERTY()
+	int32 string_count;
+
+	UPROPERTY()
+	int32 word_count;
+
+	UPROPERTY()
+	FString status;
+
+	UPROPERTY()
+	FString created_at;
+
+	UPROPERTY()
+	int32 created_at_timestamp;
+
+	/** Default constructor. */
+	FOneSkyShowImportTaskResponseData()
+		: id(-1)
+		, file()
+		, string_count(-1)
+		, word_count(-1)
+		, status("")
+		, created_at("")
+		, created_at_timestamp(-1)
+	{ }
+
+	/** Creates an uninitialized instance. */
+	FOneSkyShowImportTaskResponseData(ENoInit) { }
+};
+
+/**
+* Response from a Show Import Task query on OneSky
+* //[REDACTED_PROFILE_URL]/api-documentation-platform/blob/master/resources/import_task.md
+*/
+USTRUCT()
+struct FOneSkyShowImportTaskResponse
+{
+	GENERATED_USTRUCT_BODY()
+
+	UPROPERTY()
+	FOneSkyShowImportTaskResponseMeta meta;
+
+	UPROPERTY()
+	FOneSkyShowImportTaskResponseData data;
+
+	/** Default constructor. */
+	FOneSkyShowImportTaskResponse()
+		: meta()
+		, data()
+	{ }
+
+	/** Creates an uninitialized instance. */
+	FOneSkyShowImportTaskResponse(ENoInit) { }
+};
+
+// end SHOW IMPORT TASK 
\ No newline at end of file
diff --git a/Engine/Source/Developer/LocalizationService/Public/LocalizationServiceOperations.h b/Engine/Source/Developer/LocalizationService/Public/LocalizationServiceOperations.h
--- a/Engine/Source/Developer/LocalizationService/Public/LocalizationServiceOperations.h
+++ b/Engine/Source/Developer/LocalizationService/Public/LocalizationServiceOperations.h
@@ -52,7 +52,7 @@ protected:
 };
 
 /**
- * Operation used to commit localization(s) to localization service
+ * Operation used to download a localization target file from a localization service
  */
 class FDownloadLocalizationTargetFile : public ILocalizationServiceOperation
 {
@@ -66,7 +66,7 @@ public:
 
 	virtual FText GetInProgressString() const override
 	{ 
-		return LOCTEXT("LocalizationService_FLocalizationServicerDownloadLocalizationTargetFile", "Getting Localization Target File from localization Service..."); 
+		return LOCTEXT("LocalizationService_FLocalizationServiceDownloadLocalizationTargetFile", "Getting Localization Target File from localization Service..."); 
 	}
 
 
@@ -126,79 +126,96 @@ protected:
 	FText OutErrorText;
 };
 
-///**
-// * Operation used to commit localization(s) to localization service
-// */
-//class FCommitTranslations : public ILocalizationServiceOperation
-//{
-//public:
-//	// ILocalizationServiceOperation interface
-//	virtual FName GetName() const override
-//	{
-//		return "CommitTranslations";
-//	}
-//
-//	virtual FText GetInProgressString() const override
-//	{ 
-//		return LOCTEXT("LocalizationService_FCommit", "Committing translation(s) to localization Service..."); 
-//	}
-//
-//
-//	void SetSuccessMessage( const FText& InSuccessMessage )
-//	{
-//		SuccessMessage = InSuccessMessage;
-//	}
-//
-//	const FText& GetSuccessMessage() const
-//	{
-//		return SuccessMessage;
-//	}
-//
-//protected:
-//
-//	/** A short message listing changelist/revision we submitted, if successful */
-//	FText SuccessMessage;
-//
-//	/** The translations we are committing */
-//	TArray<TPair<FLocalizationServiceTranslationIdentifier, FString>> TranslationsToCommit;
-//};
-
-///**
-//* Operation used to update the localization service status of files
-//*/
-//class FLocalizationServiceUpdateStatus : public ILocalizationServiceOperation
-//{
-//public:
-//	FLocalizationServiceUpdateStatus()
-//		: bUpdateHistory(false)
-//	{
-//	}
-//
-//	// ILocalizationServiceOperation interface
-//	virtual FName GetName() const override
-//	{
-//		return "UpdateStatus";
-//	}
-//
-//	virtual FText GetInProgressString() const override
-//	{
-//		return LOCTEXT("LocalizationService_Update", "Updating translation(s) localization service status...");
-//	}
-//
-//	void SetUpdateHistory(bool bInUpdateHistory)
-//	{
-//		bUpdateHistory = bInUpdateHistory;
-//	}
-//
-//	bool ShouldUpdateHistory() const
-//	{
-//		return bUpdateHistory;
-//	}
-//	
-//protected:
-//	/** Whether to update history */
-//	bool bUpdateHistory;
-//};
+
+/**
+* Operation used to upload a localization file to a localization service
+*/
+class FUploadLocalizationTargetFile : public ILocalizationServiceOperation
+{
+public:
+
+	FUploadLocalizationTargetFile() : bPreserveAllText(true) {}
+
+	// ILocalizationServiceOperation interface
+	virtual FName GetName() const override
+	{
+		return "UploadLocalizationTargetFile";
+	}
+
+	virtual FText GetInProgressString() const override
+	{
+		return LOCTEXT("LocalizationService_FLocalizationServiceUploadLocalizationTargetFile", "Uploading Localization Target File to localization Service...");
+	}
+
+
+	void SetInTargetGuid(const FGuid& NewInTargetGuid)
+	{
+		InTargetGuid = NewInTargetGuid;
+	}
+
+	const FGuid& GetInTargetGuid() const
+	{
+		return InTargetGuid;
+	}
+
+	void SetInRelativeInputFilePathAndName(const FString& NewInRelativeInputFilePathAndName)
+	{
+		InRelativeInputFilePathAndName = NewInRelativeInputFilePathAndName;
+	}
+
+	const FString& GetInRelativeInputFilePathAndName() const
+	{
+		return InRelativeInputFilePathAndName;
+	}
+
+	FString GetInLocale()
+	{
+		return InLocale;
+	}
+
+	void SetInLocale(FString NewLocale)
+	{
+		InLocale = NewLocale;
+	}
+
+	FText GetOutErrorText()
+	{
+		return OutErrorText;
+	}
+
+	void SetOutErrorText(const FText& NewOutErrorText)
+	{
+		OutErrorText = NewOutErrorText;
+	}
+
+	bool GetPreserveAllText()
+	{
+		return bPreserveAllText;
+	}
+
+	void SetPreserveAllText(bool NewPreserveAllText)
+	{
+		bPreserveAllText = NewPreserveAllText;
+	}
+
+
+protected:
+
+	/** The GUID of the Localization Target */
+	FGuid InTargetGuid;
+
+	/** The locale (culture code, for example "fr" or "en-us" or "ja-jp") to download the translations for */
+	FString InLocale;
+
+	/** The path and name to the file to upload, relative to project directory */
+	FString InRelativeInputFilePathAndName;
+
+	/** Place to easily store and access error message */
+	FText OutErrorText;
+
+	/** Ask the translations service to keep all text, even if not present in the current upload file.*/
+	bool bPreserveAllText;
+};
 
 
 #undef LOCTEXT_NAMESPACE
\ No newline at end of file
diff --git a/Engine/Source/Editor/LocalizationDashboard/LocalizationDashboard.Build.cs b/Engine/Source/Editor/LocalizationDashboard/LocalizationDashboard.Build.cs
--- a/Engine/Source/Editor/LocalizationDashboard/LocalizationDashboard.Build.cs
+++ b/Engine/Source/Editor/LocalizationDashboard/LocalizationDashboard.Build.cs
@@ -46,5 +46,13 @@ public class LocalizationDashboard : ModuleRules
 				"Editor/LocalizationDashboard/Public",
 			}
 		);
+
+        CircularlyReferencedDependentModules.AddRange(
+            new string[]
+			{
+                 "MainFrame",
+                 "TranslationEditor"
+			}
+            );
 	}
 }
\ No newline at end of file
diff --git a/Engine/Source/Editor/LocalizationDashboard/Private/LocalizationDashboardModule.cpp b/Engine/Source/Editor/LocalizationDashboard/Private/LocalizationDashboardModule.cpp
--- a/Engine/Source/Editor/LocalizationDashboard/Private/LocalizationDashboardModule.cpp
+++ b/Engine/Source/Editor/LocalizationDashboard/Private/LocalizationDashboardModule.cpp
@@ -12,6 +12,7 @@
 #include "GatherTextDetailCustomizations.h"
 #include "PropertyEditorModule.h"
 #include "ILocalizationServiceModule.h"
+#include "LocalizationDashboardSettings.h"
 
 #define LOCTEXT_NAMESPACE "LocalizationDashboard"
 
@@ -152,6 +153,35 @@ public:
 		return LSP ? *LSP : nullptr;
 	}
 
+	ULocalizationTarget* GetLocalizationTargetByName(FString TargetName, bool bIsEngineTarget) override
+	{
+		if (bIsEngineTarget)
+		{
+			ULocalizationTargetSet* EngineTargetSet = ULocalizationDashboardSettings::GetEngineTargetSet();
+			for (ULocalizationTarget* Target : EngineTargetSet->TargetObjects)
+			{
+				if (Target->Settings.Name == TargetName)
+				{
+					return Target;
+				}
+			}
+		}
+		else
+		{
+			ULocalizationTargetSet* EngineTargetSet = ULocalizationDashboardSettings::GetEngineTargetSet();
+			for (ULocalizationTarget* Target : EngineTargetSet->TargetObjects)
+			{
+				if (Target->Settings.Name == TargetName)
+				{
+					return Target;
+				}
+			}
+		}
+
+		return nullptr;
+	}
+
+
 private:
 	TArray<ILocalizationServiceProvider*> ServiceProviders;
 	FName CurrentServiceProviderName;
@@ -159,4 +189,4 @@ private:
 
 IMPLEMENT_MODULE(FLocalizationDashboardModule, LocalizationDashboard);
 
-#undef LOCTEXT_NAMESPACE
+#undef LOCTEXT_NAMESPACE
\ No newline at end of file
diff --git a/Engine/Source/Editor/LocalizationDashboard/Public/ILocalizationDashboardModule.h b/Engine/Source/Editor/LocalizationDashboard/Public/ILocalizationDashboardModule.h
--- a/Engine/Source/Editor/LocalizationDashboard/Public/ILocalizationDashboardModule.h
+++ b/Engine/Source/Editor/LocalizationDashboard/Public/ILocalizationDashboardModule.h
@@ -3,6 +3,7 @@
 #pragma once
 
 #include "ModuleInterface.h"
+#include "LocalizationTargetTypes.h"
 
 class ILocalizationServiceProvider;
 
@@ -28,4 +29,6 @@ public:
 	{
 		return FModuleManager::LoadModuleChecked<ILocalizationDashboardModule>("LocalizationDashboard");
 	}
+
+	virtual ULocalizationTarget* GetLocalizationTargetByName(FString TargetName, bool bIsEngineTarget) = 0;
 };
diff --git a/Engine/Source/Editor/TranslationEditor/Private/TranslationDataManager.cpp b/Engine/Source/Editor/TranslationEditor/Private/TranslationDataManager.cpp
--- a/Engine/Source/Editor/TranslationEditor/Private/TranslationDataManager.cpp
+++ b/Engine/Source/Editor/TranslationEditor/Private/TranslationDataManager.cpp
@@ -11,6 +11,10 @@
 #include "JsonInternationalizationManifestSerializer.h"
 #include "SNotificationList.h"
 #include "NotificationManager.h"
+#include "PortableObjectFormatDOM.h"
+#include "ILocalizationServiceModule.h"
+#include "Editor/LocalizationDashboard/Public/ILocalizationDashboardModule.h"
+#include "Editor/LocalizationDashboard/Public/LocalizationTargetTypes.h"
 
 
 DEFINE_LOG_CATEGORY_STATIC(LogTranslationEditor, Log, All);
@@ -805,7 +809,7 @@ void FTranslationDataManager::UnloadHistoryInformation()
 	}
 }
 
-bool FTranslationDataManager::SaveSelectedTranslations(TArray<UTranslationUnit*> TranslationUnitsToSave)
+bool FTranslationDataManager::SaveSelectedTranslations(TArray<UTranslationUnit*> TranslationUnitsToSave, bool bSaveChangesToTranslationService)
 {
 	bool bSucceeded = true;
 	
@@ -837,34 +841,89 @@ bool FTranslationDataManager::SaveSelectedTranslations(TArray<UTranslationUnit*>
 		FString ManifestPath = FPaths::GetPath(ArchiveFilePath);
 		FString ArchiveFullPath = ArchiveFilePath / ManifestAndArchiveName + ".archive";
 		FString ManifestFullPath = ManifestPath / ManifestAndArchiveName + ".manifest";
+		FString EngineFullPath = FPaths::ConvertRelativePathToFull(FPaths::EngineContentDir());
+		bool IsEngineManifest = false;
+		if (ManifestFullPath.StartsWith(EngineFullPath))
+		{
+			IsEngineManifest = true;
+		}
 
 		if (FPaths::FileExists(ManifestFullPath) && FPaths::FileExists(ArchiveFullPath))
 		{
 			TSharedRef<FTranslationDataManager> DataManager = MakeShareable(new FTranslationDataManager(ManifestFullPath, ArchiveFullPath));
 
-			TArray<UTranslationUnit*>& TranslationsArray = DataManager->GetAllTranslationsArray();
-			TSharedPtr<TArray<UTranslationUnit*>> EditedItems = Item.Value;
-
-			// For each edited item belonging to this manifest/archive pair
-			for (auto EditedItemIt = EditedItems->CreateIterator(); EditedItemIt; ++EditedItemIt)
+			if (DataManager->GetLoadedSuccessfully())
 			{
-				UTranslationUnit* EditedItem = *EditedItemIt;
+				FPortableObjectFormatDOM PortableObjectDom;
+				PortableObjectDom.SetProjectName(ManifestAndArchiveName);
+				PortableObjectDom.SetLanguage(CultureName);
+				PortableObjectDom.CreateNewHeader();
 
-				// Search all translations for the one that matches this FText
-				for (UTranslationUnit* Translation : TranslationsArray)
+				TArray<UTranslationUnit*>& TranslationsArray = DataManager->GetAllTranslationsArray();
+				TSharedPtr<TArray<UTranslationUnit*>> EditedItems = Item.Value;
+
+				// For each edited item belonging to this manifest/archive pair
+				for (auto EditedItemIt = EditedItems->CreateIterator(); EditedItemIt; ++EditedItemIt)
 				{
-					// If namespace matches...
-					if (Translation->Namespace == EditedItem->Namespace)
+					UTranslationUnit* EditedItem = *EditedItemIt;
+
+					// Search all translations for the one that matches this FText
+					for (UTranslationUnit* Translation : TranslationsArray)
 					{
-						// And source matches
-						if (Translation->Source == EditedItem->Source)
+						// If namespace matches...
+						if (Translation->Namespace == EditedItem->Namespace)
 						{
-							// Update the translation in TranslationDataManager, and finish searching these translations
-							Translation->Translation = EditedItem->Translation;
-							break;
+							// And source matches
+							if (Translation->Source == EditedItem->Source)
+							{
+								// Update the translation in TranslationDataManager, and finish searching these translations
+								Translation->Translation = EditedItem->Translation;
+
+								TSharedRef<FPortableObjectEntry> NewEntry = MakeShareable(new FPortableObjectEntry);
+								for (FTranslationContextInfo ContextInfo : Translation->Contexts)
+								{
+									NewEntry->ExtractedComments.Add(ContextInfo.Key);
+									NewEntry->ReferenceComments.Add(ContextInfo.Context);
+								}
+
+								NewEntry->MsgCtxt = Translation->Namespace;
+								NewEntry->MsgId = Translation->Source;
+								NewEntry->MsgStr.Add(Translation->Translation);
+								PortableObjectDom.AddEntry(NewEntry);
+
+								break;
+							}
 						}
 					}
 				}
+
+				if (bSaveChangesToTranslationService)
+				{
+					FString UploadFilePath = FPaths::GameSavedDir() / "Temp" / CultureName / ManifestAndArchiveName + ".po";
+					FFileHelper::SaveStringToFile(PortableObjectDom.ToString(), *UploadFilePath);
+
+					ULocalizationTarget* Target = ILocalizationDashboardModule::Get().GetLocalizationTargetByName(ManifestAndArchiveName, IsEngineManifest);
+					FGuid LocalizationTargetGuid;
+					if (Target)
+					{
+						LocalizationTargetGuid = Target->Settings.Guid;
+					}
+
+					ILocalizationServiceProvider& Provider = ILocalizationServiceModule::Get().GetProvider();
+					TSharedRef<FUploadLocalizationTargetFile, ESPMode::ThreadSafe> UploadTargetFileOp = ILocalizationServiceOperation::Create<FUploadLocalizationTargetFile>();
+					UploadTargetFileOp->SetInTargetGuid(LocalizationTargetGuid);
+					UploadTargetFileOp->SetInLocale(CultureName);
+					FPaths::MakePathRelativeTo(UploadFilePath, *FPaths::GameDir());
+					UploadTargetFileOp->SetInRelativeInputFilePathAndName(UploadFilePath);
+					UploadTargetFileOp->SetPreserveAllText(true);
+
+					Provider.Execute(UploadTargetFileOp, TArray<FLocalizationServiceTranslationIdentifier>(), ELocalizationServiceOperationConcurrency::Asynchronous, FLocalizationServiceOperationComplete::CreateStatic(&FTranslationDataManager::SaveSelectedTranslationsToTranslationServiceCallback));
+				}
+
+			}
+			else
+			{
+				bSucceeded = false;
 			}
 
 			// Save the data to file, and preview in editor
@@ -880,4 +939,57 @@ bool FTranslationDataManager::SaveSelectedTranslations(TArray<UTranslationUnit*>
 	return bSucceeded;
 }
 
+void FTranslationDataManager::SaveSelectedTranslationsToTranslationServiceCallback(const FLocalizationServiceOperationRef& Operation, ELocalizationServiceOperationCommandResult::Type Result)
+{
+	TSharedPtr<FUploadLocalizationTargetFile, ESPMode::ThreadSafe> DownloadLocalizationTargetOp = StaticCastSharedRef<FUploadLocalizationTargetFile>(Operation);
+	bool bError = !(Result == ELocalizationServiceOperationCommandResult::Succeeded);
+	FText ErrorText = FText::GetEmpty();
+	FGuid InTargetGuid;
+	FString InLocale;
+	FString InRelativeOutputFilePathAndName;
+	FString TargetName = "";
+	FString TargetPath = "";
+	FString CultureName = "";
+	if (DownloadLocalizationTargetOp.IsValid())
+	{
+		ErrorText = DownloadLocalizationTargetOp->GetOutErrorText();
+		InTargetGuid = DownloadLocalizationTargetOp->GetInTargetGuid();
+		InLocale = DownloadLocalizationTargetOp->GetInLocale();
+		InRelativeOutputFilePathAndName = DownloadLocalizationTargetOp->GetInRelativeInputFilePathAndName();
+		TargetName = FPaths::GetBaseFilename(InRelativeOutputFilePathAndName);
+		TargetPath = FPaths::GetPath(InRelativeOutputFilePathAndName);
+		CultureName = FPaths::GetBaseFilename(TargetPath);
+	}
+
+	// Try to get display name
+	FInternationalization& I18N = FInternationalization::Get();
+	FCulturePtr CulturePtr = I18N.GetCulture(CultureName);
+	FString CultureDisplayName = CultureName;
+	if (CulturePtr.IsValid())
+	{
+		CultureName = CulturePtr->GetDisplayName();
+	}
+
+	if (!bError && ErrorText.IsEmpty())
+	{
+		FText SuccessText = FText::Format(LOCTEXT("SaveSelectedTranslationsToTranslationServiceSuccess", "{0} translations for {1} target uploaded for processing to Translation Service."), FText::FromString(CultureDisplayName), FText::FromString(TargetName));
+		FMessageLog TranslationEditorMessageLog("TranslationEditor");
+		TranslationEditorMessageLog.Info(SuccessText);
+		TranslationEditorMessageLog.Notify(SuccessText, EMessageSeverity::Info, true);
+	}
+	else
+	{
+		if (ErrorText.IsEmpty())
+		{
+			ErrorText = LOCTEXT("DownloadLatestFromLocalizationServiceUnspecifiedError", "An unspecified error occured when trying download and import from the Localization Service.");
+		}
+
+		FText ErrorNotify = FText::Format(LOCTEXT("SaveSelectedTranslationsToTranslationServiceFail", "{0} translations for {1} target failed to save to Translation Service!"), FText::FromString(CultureDisplayName), FText::FromString(TargetName));
+		FMessageLog TranslationEditorMessageLog("TranslationEditor");
+		TranslationEditorMessageLog.Error(ErrorNotify);
+		TranslationEditorMessageLog.Error(ErrorText);
+		TranslationEditorMessageLog.Notify(ErrorNotify);
+	}
+}
+
 #undef LOCTEXT_NAMESPACE
diff --git a/Engine/Source/Editor/TranslationEditor/Private/TranslationDataManager.h b/Engine/Source/Editor/TranslationEditor/Private/TranslationDataManager.h
--- a/Engine/Source/Editor/TranslationEditor/Private/TranslationDataManager.h
+++ b/Engine/Source/Editor/TranslationEditor/Private/TranslationDataManager.h
@@ -9,6 +9,7 @@ class UTranslationUnit;
 
 #include "JsonInternationalizationArchiveSerializer.h"
 #include "JsonInternationalizationManifestSerializer.h"
+#include "ILocalizationServiceProvider.h"
 
 class FTranslationDataManager : public TSharedFromThis<FTranslationDataManager>
 {
@@ -74,7 +75,10 @@ public:
 	void UnloadHistoryInformation();
 
 	/** Save the specified translations */
-	static bool SaveSelectedTranslations(TArray<UTranslationUnit*> TranslationUnitsToSave);
+	static bool SaveSelectedTranslations(TArray<UTranslationUnit*> TranslationUnitsToSave, bool bSaveChangesToTranslationService = false);
+
+	/** Save the specified translations */
+	static void SaveSelectedTranslationsToTranslationServiceCallback(const FLocalizationServiceOperationRef& Operation, ELocalizationServiceOperationCommandResult::Type Result);
 
 	/** Whether or not we successfully loaded the .manifest and .archive */
 	bool GetLoadedSuccessfully()
diff --git a/Engine/Source/Editor/TranslationEditor/Private/TranslationEditorMenu.cpp b/Engine/Source/Editor/TranslationEditor/Private/TranslationEditorMenu.cpp
--- a/Engine/Source/Editor/TranslationEditor/Private/TranslationEditorMenu.cpp
+++ b/Engine/Source/Editor/TranslationEditor/Private/TranslationEditorMenu.cpp
@@ -21,7 +21,10 @@ void FTranslationEditorMenu::FillTranslationMenu( FMenuBuilder& MenuBuilder/*, F
 		MenuBuilder.AddMenuEntry( FTranslationEditorCommands::Get().ExportToPortableObjectFormat );
 		MenuBuilder.AddMenuEntry( FTranslationEditorCommands::Get().ImportFromPortableObjectFormat );
 		MenuBuilder.AddMenuEntry( FTranslationEditorCommands::Get().OpenSearchTab );
-		MenuBuilder.AddMenuEntry( FTranslationEditorCommands::Get().OpenTranslationPicker );
+		if (GetDefault<UEditorExperimentalSettings>()->bEnableTranslationPicker)
+		{
+			MenuBuilder.AddMenuEntry(FTranslationEditorCommands::Get().OpenTranslationPicker);
+		}
 	}
 	MenuBuilder.EndSection();
 }
@@ -80,7 +83,10 @@ void FTranslationEditorMenu::SetupTranslationEditorToolbar( TSharedPtr< FExtende
 				FTranslationEditorCommands::Get().ImportFromPortableObjectFormat, "ImportFromPortableObjectFormat", TAttribute<FText>(), TAttribute<FText>(), FSlateIcon(FEditorStyle::GetStyleSetName(), "TranslationEditor.Import"));
 			ToolbarBuilder.AddToolBarButton(
 				FTranslationEditorCommands::Get().OpenSearchTab, "OpenSearchTab", TAttribute<FText>(), TAttribute<FText>(), FSlateIcon(FEditorStyle::GetStyleSetName(), "TranslationEditor.Search"));
-			ToolbarBuilder.AddWidget(SNew(STranslationWidgetPicker));
+			if (GetDefault<UEditorExperimentalSettings>()->bEnableTranslationPicker)
+			{
+				ToolbarBuilder.AddWidget(SNew(STranslationWidgetPicker));
+			}
 		}
 	};
 
diff --git a/Engine/Source/Editor/TranslationEditor/Private/TranslationPickerEditWindow.cpp b/Engine/Source/Editor/TranslationEditor/Private/TranslationPickerEditWindow.cpp
--- a/Engine/Source/Editor/TranslationEditor/Private/TranslationPickerEditWindow.cpp
+++ b/Engine/Source/Editor/TranslationEditor/Private/TranslationPickerEditWindow.cpp
@@ -4,19 +4,37 @@
 #include "TranslationPickerEditWindow.h"
 #include "Editor/Documentation/Public/SDocumentationToolTip.h"
 #include "TranslationUnit.h"
+#include "ILocalizationServiceModule.h"
 
 #define LOCTEXT_NAMESPACE "TranslationPicker"
 
+TSharedPtr<FTranslationPickerSettingsManager> FTranslationPickerSettingsManager::TranslationPickerSettingsManagerInstance;
+
 // Default dimensions of the Translation Picker edit window (floating window also uses these sizes, so it matches roughly)
 const int32 STranslationPickerEditWindow::DefaultEditWindowWidth = 500;
 const int32 STranslationPickerEditWindow::DefaultEditWindowHeight = 500;
 
+UTranslationPickerSettings::UTranslationPickerSettings(const FObjectInitializer& ObjectInitializer)
+: Super(ObjectInitializer)
+{
+
+}
+
 void STranslationPickerEditWindow::Construct(const FArguments& InArgs)
 {
 	ParentWindow = InArgs._ParentWindow;
 	PickedTexts = InArgs._PickedTexts;
 	WindowContents = SNew(SBox);
 	TSharedRef<SVerticalBox> TextsBox = SNew(SVerticalBox);
+	UTranslationPickerSettings* TranslationPickerSettings = FTranslationPickerSettingsManager::Get()->GetSettings();
+
+	bool bShowLocServiceCheckbox = ILocalizationServiceModule::Get().GetProvider().IsEnabled();
+
+	if (!FParse::Param(FCommandLine::Get(), TEXT("AllowTranslationPickerSubmissionsToOneSky")))
+	{
+		bShowLocServiceCheckbox = false;
+		TranslationPickerSettings->bSubmitTranslationPickerChangesToLocalizationService = false;
+	}
 
 	// Add a new Translation Picker Edit Widget for each picked text
 	for (FText PickedText : PickedTexts)
@@ -79,38 +97,79 @@ void STranslationPickerEditWindow::Construct(const FArguments& InArgs)
 			]
 			+ SVerticalBox::Slot()
 			.AutoHeight()
-			.HAlign(HAlign_Right)
+			.HAlign(HAlign_Fill)
 			.Padding(DefaultPadding)
 			[
-				SNew(SUniformGridPanel)
-				.SlotPadding(FEditorStyle::GetMargin("StandardDialog.SlotPadding"))
-				.MinDesiredSlotWidth(FEditorStyle::GetFloat("StandardDialog.MinDesiredSlotWidth"))
-				.MinDesiredSlotHeight(FEditorStyle::GetFloat("StandardDialog.MinDesiredSlotHeight"))
-				+ SUniformGridPanel::Slot(0, 0)
+				SNew(SVerticalBox)
+				+ SVerticalBox::Slot()
+				.AutoHeight()
+				.HAlign(HAlign_Left)
+				.Padding(DefaultPadding)
 				[
-					SNew(SButton)
-					.HAlign(HAlign_Center)
-					.ContentPadding(FEditorStyle::GetMargin("StandardDialog.ContentPadding"))
-					.OnClicked(this, &STranslationPickerEditWindow::SaveAllAndClose)
+					SNew(SHorizontalBox)
+					.Visibility(bShowLocServiceCheckbox ? EVisibility::Visible : EVisibility::Collapsed)
+					+ SHorizontalBox::Slot()
+					.Padding(FMargin(3, 3, 3, 3))
+					.VAlign(VAlign_Center)
+					.AutoWidth()
 					[
-						SNew(SHorizontalBox)
-						+ SHorizontalBox::Slot()
-						.Padding(FMargin(0, 0, 3, 0))
-						.VAlign(VAlign_Center)
-						.AutoWidth()
-						[
-							SNew(STextBlock).Text(LOCTEXT("SaveAllAndClose", "Save all and close"))
-						]
+						SNew(SCheckBox)
+						.HAlign(HAlign_Center)
+						.IsChecked(TranslationPickerSettings->bSubmitTranslationPickerChangesToLocalizationService ? ECheckBoxState::Checked : ECheckBoxState::Unchecked)
+						.ToolTipText(LOCTEXT("SubmitTranslationPickerChangesToLocalizationServiceToolTip", "Submit changes to localization service"))
+						.OnCheckStateChanged_Lambda([&](ECheckBoxState CheckedState)
+						{
+							UTranslationPickerSettings* TranslationPickerSettingsLocal = FTranslationPickerSettingsManager::Get()->GetSettings();
+							TranslationPickerSettingsLocal->bSubmitTranslationPickerChangesToLocalizationService = CheckedState == ECheckBoxState::Checked;
+							TranslationPickerSettingsLocal->SaveConfig();
+						}
+						)
+					]
+					+ SHorizontalBox::Slot()
+					.Padding(FMargin(0, 0, 3, 0))
+					.VAlign(VAlign_Center)
+					.AutoWidth()
+					[
+						SNew(STextBlock)
+						.Text(LOCTEXT("SubmitTranslationPickerChangesToLocalizationService", "Save to Localization Service"))
+						.ToolTipText(LOCTEXT("SubmitTranslationPickerChangesToLocalizationServiceToolTip", "Submit changes to localization service"))
 					]
 				]
-				+ SUniformGridPanel::Slot(1, 0)
+				+ SVerticalBox::Slot()
+				.AutoHeight()
+				.HAlign(HAlign_Right)
+				.Padding(DefaultPadding)
 				[
-					SNew(SButton)
-					.HAlign(HAlign_Center)
-					.ContentPadding(FEditorStyle::GetMargin("StandardDialog.ContentPadding"))
-					.OnClicked(this, &STranslationPickerEditWindow::Close)
-					.VAlign(VAlign_Center)
-					.Text(LOCTEXT("CancelButton", "Cancel"))
+					SNew(SUniformGridPanel)
+					.SlotPadding(FEditorStyle::GetMargin("StandardDialog.SlotPadding"))
+					.MinDesiredSlotWidth(FEditorStyle::GetFloat("StandardDialog.MinDesiredSlotWidth"))
+					.MinDesiredSlotHeight(FEditorStyle::GetFloat("StandardDialog.MinDesiredSlotHeight"))
+					+ SUniformGridPanel::Slot(0, 0)
+					[
+						SNew(SButton)
+						.HAlign(HAlign_Center)
+						.ContentPadding(FEditorStyle::GetMargin("StandardDialog.ContentPadding"))
+						.OnClicked(this, &STranslationPickerEditWindow::SaveAllAndClose)
+						[
+							SNew(SHorizontalBox)
+							+ SHorizontalBox::Slot()
+							.Padding(FMargin(0, 0, 3, 0))
+							.VAlign(VAlign_Center)
+							.AutoWidth()
+							[
+								SNew(STextBlock).Text(LOCTEXT("SaveAllAndClose", "Save all and close"))
+							]
+						]
+					]
+					+ SUniformGridPanel::Slot(1, 0)
+					[
+						SNew(SButton)
+						.HAlign(HAlign_Center)
+						.ContentPadding(FEditorStyle::GetMargin("StandardDialog.ContentPadding"))
+						.OnClicked(this, &STranslationPickerEditWindow::Close)
+						.VAlign(VAlign_Center)
+						.Text(LOCTEXT("CancelButton", "Cancel"))
+					]
 				]
 			]
 		]
@@ -161,8 +220,9 @@ FReply STranslationPickerEditWindow::SaveAllAndClose()
 
 	if (TempArray.Num() > 0)
 	{
+		UTranslationPickerSettings* TranslationPickerSettings = FTranslationPickerSettingsManager::Get()->GetSettings();
 		// Save the data via translation data manager
-		FTranslationDataManager::SaveSelectedTranslations(TempArray);
+		FTranslationDataManager::SaveSelectedTranslations(TempArray, ILocalizationServiceModule::Get().GetProvider().IsEnabled() && TranslationPickerSettings->bSubmitTranslationPickerChangesToLocalizationService);
 	}
 
 	Close();
@@ -197,6 +257,19 @@ void STranslationPickerEditWidget::Construct(const FArguments& InArgs)
 		}
 	}
 
+	FString ArchiveFilePath = FPaths::GetPath(LocresFullPath);
+	FString LocResCultureName = FPaths::GetBaseFilename(ArchiveFilePath);
+
+	FInternationalization& I18N = FInternationalization::Get();
+	FString CurrentCultureName = I18N.GetCurrentCulture()->GetName();
+
+	// Make sure locres path matches the current culture (can change if using the culture=__ console command)
+	if (CurrentCultureName != LocResCultureName)
+	{
+		FString BasePath = FPaths::GetPath(ArchiveFilePath);
+		LocresFullPath = BasePath / CurrentCultureName / ManifestAndArchiveNameString + ".locres";
+	}
+
 	FText Namespace = FText::FromString(NamespaceString.Get(TEXT("")));
 	FText Key = FText::FromString(KeyString.Get(TEXT("")));
 	FText Source = SourceString != nullptr ? FText::FromString(*SourceString) : FText::GetEmpty();
@@ -392,11 +465,12 @@ FReply STranslationPickerEditWidget::SaveAndPreview()
 {
 	// Update translation string from entered text
 	TranslationUnit->Translation = TextBox->GetText().ToString();
+	UTranslationPickerSettings* TranslationPickerSettings = FTranslationPickerSettingsManager::Get()->GetSettings();
 
 	// Save the data via translation data manager
 	TArray<UTranslationUnit*> TempArray;
 	TempArray.Add(TranslationUnit);
-	FTranslationDataManager::SaveSelectedTranslations(TempArray);
+	FTranslationDataManager::SaveSelectedTranslations(TempArray, ILocalizationServiceModule::Get().GetProvider().IsEnabled() && TranslationPickerSettings->bSubmitTranslationPickerChangesToLocalizationService);
 
 	return FReply::Handled();
 }
diff --git a/Engine/Source/Editor/TranslationEditor/Private/TranslationPickerEditWindow.h b/Engine/Source/Editor/TranslationEditor/Private/TranslationPickerEditWindow.h
--- a/Engine/Source/Editor/TranslationEditor/Private/TranslationPickerEditWindow.h
+++ b/Engine/Source/Editor/TranslationEditor/Private/TranslationPickerEditWindow.h
@@ -2,9 +2,68 @@
 
 #pragma once
 #include "SCompoundWidget.h"
+#include "TranslationPickerEditWindow.generated.h"
 
 #define LOCTEXT_NAMESPACE "TranslationPicker"
 
+UCLASS(config = TranslationPickerSettings)
+class UTranslationPickerSettings : public UObject
+{
+	GENERATED_UCLASS_BODY()
+
+public:
+	/** Bool submit translation picker changes to Localization Service */
+	UPROPERTY(config)
+	bool bSubmitTranslationPickerChangesToLocalizationService;
+};
+
+class FTranslationPickerSettingsManager
+{
+	FTranslationPickerSettingsManager()
+	{
+		TranslationPickerSettingsObject = NewObject<UTranslationPickerSettings>();
+		TranslationPickerSettingsObject->LoadConfig();
+	}
+
+public:
+	void SaveSettings()
+	{
+		TranslationPickerSettingsObject->SaveConfig();
+	}
+
+	void LoadSettings()
+	{
+		TranslationPickerSettingsObject->LoadConfig();
+	}
+
+	UTranslationPickerSettings* GetSettings()
+	{
+		return TranslationPickerSettingsObject;
+	}
+
+	/**
+	* Gets a reference to the translation picker manager instance.
+	*
+	* @return A reference to the translation picker manager instance.
+	*/
+	static inline TSharedPtr<FTranslationPickerSettingsManager>& Get()
+	{
+		if (!TranslationPickerSettingsManagerInstance.IsValid())
+		{
+			TranslationPickerSettingsManagerInstance = MakeShareable(new FTranslationPickerSettingsManager());
+		}
+
+		return TranslationPickerSettingsManagerInstance;
+	}
+
+private:
+
+	static TSharedPtr<FTranslationPickerSettingsManager> TranslationPickerSettingsManagerInstance;
+
+	/** Used to load and store settings for the Translation Picker */
+	UTranslationPickerSettings* TranslationPickerSettingsObject;
+};
+
 /** Translation picker edit Widget to handle the display and editing of a single selected FText */
 class STranslationPickerEditWidget : public SCompoundWidget
 {
@@ -86,6 +145,8 @@ private:
 
 	/** Save all translations and close */
 	FReply SaveAllAndClose();
+
+
 };
 
 #undef LOCTEXT_NAMESPACE
\ No newline at end of file
diff --git a/Engine/Source/Editor/TranslationEditor/TranslationEditor.Build.cs b/Engine/Source/Editor/TranslationEditor/TranslationEditor.Build.cs
--- a/Engine/Source/Editor/TranslationEditor/TranslationEditor.Build.cs
+++ b/Engine/Source/Editor/TranslationEditor/TranslationEditor.Build.cs
@@ -36,6 +36,7 @@ public class TranslationEditor : ModuleRules
                 "MessageLog",
                 "Documentation",
                 "LocalizationService",
+                "LocalizationDashboard",
 			}
 		);
 
