OneSky Localization Plugin: Label this plugin as beta for now

diff --git a/Engine/Plugins/Developer/OneSkyLocalizationService/OneSkyLocalizationService.uplugin b/Engine/Plugins/Developer/OneSkyLocalizationService/OneSkyLocalizationService.uplugin
--- a/Engine/Plugins/Developer/OneSkyLocalizationService/OneSkyLocalizationService.uplugin
+++ b/Engine/Plugins/Developer/OneSkyLocalizationService/OneSkyLocalizationService.uplugin
@@ -3,13 +3,14 @@
 	
 	"FriendlyName" : "OneSky",
 	"Version" : 1,
-	"VersionName" : "1.0",
+	"VersionName" : "0.5",
 	"CreatedBy" : "Epic Games, Inc.",
 	"CreatedByURL" : "http://epicgames.com",
 	"EngineVersion" : "4.8.0",
 	"Description" : "OneSky localization service",
 	"Category" : "Editor.Localization Service",
 	"EnabledByDefault" : false,
+    "IsBetaVersion": true,
 	
 	"Modules" :
 	[
