OneSky localization service plugin: adding some missing locks and checks for thread safety.

diff --git a/Engine/Plugins/Developer/OneSkyLocalizationService/Source/OneSkyLocalizationService/Private/OneSkyLocalizationServiceSettings.cpp b/Engine/Plugins/Developer/OneSkyLocalizationService/Source/OneSkyLocalizationService/Private/OneSkyLocalizationServiceSettings.cpp
--- a/Engine/Plugins/Developer/OneSkyLocalizationService/Source/OneSkyLocalizationService/Private/OneSkyLocalizationServiceSettings.cpp
+++ b/Engine/Plugins/Developer/OneSkyLocalizationService/Source/OneSkyLocalizationService/Private/OneSkyLocalizationServiceSettings.cpp
@@ -160,6 +160,7 @@ FOneSkyConnectionInfo FOneSkyLocalizationServiceSettings::GetConnectionInfo() co
 
 FOneSkyLocalizationTargetSetting* FOneSkyLocalizationServiceSettings::GetSettingsForTarget(FGuid TargetGuid, bool bCreateIfNotFound /* = false */)
 {
+	FScopeLock ScopeLock(&CriticalSection);
 	TargetSettingsObject->LoadConfig();
 
 	FOneSkyLocalizationTargetSetting* Settings = TargetSettingsObject->TargetSettings.FindByPredicate([&](const FOneSkyLocalizationTargetSetting& Setting)
@@ -180,6 +181,7 @@ FOneSkyLocalizationTargetSetting* FOneSkyLocalizationServiceSettings::GetSetting
 
 void FOneSkyLocalizationServiceSettings::SetSettingsForTarget(FGuid TargetGuid, int32 ProjectId, FString FileName)
 {
+	FScopeLock ScopeLock(&CriticalSection);
 	FOneSkyLocalizationTargetSetting* Settings = TargetSettingsObject->TargetSettings.FindByPredicate([&](const FOneSkyLocalizationTargetSetting& Setting)
 	{
 		return Setting.TargetGuid == TargetGuid;
diff --git a/Engine/Plugins/Developer/OneSkyLocalizationService/Source/OneSkyLocalizationService/Private/OneSkyLocalizationServiceSettings.h b/Engine/Plugins/Developer/OneSkyLocalizationService/Source/OneSkyLocalizationService/Private/OneSkyLocalizationServiceSettings.h
--- a/Engine/Plugins/Developer/OneSkyLocalizationService/Source/OneSkyLocalizationService/Private/OneSkyLocalizationServiceSettings.h
+++ b/Engine/Plugins/Developer/OneSkyLocalizationService/Source/OneSkyLocalizationService/Private/OneSkyLocalizationServiceSettings.h
@@ -45,7 +45,7 @@ public:
 	* The array of settings for each target of this project 
 	*/
 	UPROPERTY(config)
-		TArray<struct FOneSkyLocalizationTargetSetting> TargetSettings;
+	TArray<struct FOneSkyLocalizationTargetSetting> TargetSettings;
 
 };
 
@@ -87,12 +87,14 @@ public:
 	/** Set whether or not to save the secret API key (WARNING: saved unencrypted) */
 	void SetSaveSecretKey(bool InSaveSecretKey)
 	{
+		check(IsInGameThread());
 		bSaveSecretKey = InSaveSecretKey;
 	}
 
 	/** Get whether or not to save the secret API Key */
 	bool GetSaveSecretKey()
 	{
+		check(IsInGameThread());
 		return bSaveSecretKey;
 	}
 
