OneSky plugin: Extend Loc Dashboard toolbar with OneSky buttons for "Import all cultures/targets from OneSky" and "Export all cultures/targets to OneSky"

diff --git a/Engine/Plugins/Developer/OneSkyLocalizationService/Source/OneSkyLocalizationService/OneSkyLocalizationService.Build.cs b/Engine/Plugins/Developer/OneSkyLocalizationService/Source/OneSkyLocalizationService/OneSkyLocalizationService.Build.cs
--- a/Engine/Plugins/Developer/OneSkyLocalizationService/Source/OneSkyLocalizationService/OneSkyLocalizationService.Build.cs
+++ b/Engine/Plugins/Developer/OneSkyLocalizationService/Source/OneSkyLocalizationService/OneSkyLocalizationService.Build.cs
@@ -10,6 +10,7 @@ public class OneSkyLocalizationService : ModuleRules
             new string[] {
 				"Core",
 				"CoreUObject",
+                "Engine",
                 "InputCore",
 				"Slate",
 				"SlateCore",
@@ -19,6 +20,8 @@ public class OneSkyLocalizationService : ModuleRules
                 "Internationalization",
                 "HTTP",
                 "Serialization",
+				"Localization",
+				"MainFrame",
 			}
 		);
 
diff --git a/Engine/Plugins/Developer/OneSkyLocalizationService/Source/OneSkyLocalizationService/Private/OneSkyLocalizationServiceProvider.cpp b/Engine/Plugins/Developer/OneSkyLocalizationService/Source/OneSkyLocalizationService/Private/OneSkyLocalizationServiceProvider.cpp
--- a/Engine/Plugins/Developer/OneSkyLocalizationService/Source/OneSkyLocalizationService/Private/OneSkyLocalizationServiceProvider.cpp
+++ b/Engine/Plugins/Developer/OneSkyLocalizationService/Source/OneSkyLocalizationService/Private/OneSkyLocalizationServiceProvider.cpp
@@ -18,16 +18,49 @@
 #include "Editor/PropertyEditor/Public/IDetailsView.h"
 
 #endif
+#include "Editor/Localization/Public/LocalizationTargetTypes.h"
+#include "LocalizationCommandletTasks.h"
+#include "ILocalizationServiceProvider.h"
+#include "LocalizationModule.h"
+#include "ModuleManager.h"
+#include "IMainFrameModule.h"
 
 static FName ProviderName("OneSky");
 
 #define LOCTEXT_NAMESPACE "OneSkyLocalizationService"
 
+class FOneSkyLocalizationTargetEditorCommands : public TCommands<FOneSkyLocalizationTargetEditorCommands>
+{
+public:
+	FOneSkyLocalizationTargetEditorCommands()
+		: TCommands<FOneSkyLocalizationTargetEditorCommands>("OneSkyLocalizationTargetEditor", NSLOCTEXT("OneSky", "OneSkyLocalizationTargetEditor", "OneSky Localization Target Editor"), NAME_None, FEditorStyle::GetStyleSetName())
+	{
+		}
+
+	TSharedPtr<FUICommandInfo> ImportAllCulturesForTargetFromOneSky;
+	TSharedPtr<FUICommandInfo> ExportAllCulturesForTargetToOneSky;
+	TSharedPtr<FUICommandInfo> ImportAllTargetsAllCulturesForTargetSetFromOneSky;
+	TSharedPtr<FUICommandInfo> ExportAllTargetsAllCulturesForTargetSetFromOneSky;
+
+	/** Initialize commands */
+	virtual void RegisterCommands() override;
+};
+
+void FOneSkyLocalizationTargetEditorCommands::RegisterCommands()
+{
+	UI_COMMAND(ImportAllCulturesForTargetFromOneSky, "Import All Cultures from OneSky", "Imports translations for all cultures of this target to OneSky.", EUserInterfaceActionType::Button, FInputChord());
+	UI_COMMAND(ExportAllCulturesForTargetToOneSky, "Export All Cultures to OneSky", "Exports translations for all cultures of this target to OneSky.", EUserInterfaceActionType::Button, FInputChord());
+	UI_COMMAND(ImportAllTargetsAllCulturesForTargetSetFromOneSky, "Import All Targets from OneSky", "Imports translations for all cultures of all targets of this target set to OneSky.", EUserInterfaceActionType::Button, FInputChord());
+	UI_COMMAND(ExportAllTargetsAllCulturesForTargetSetFromOneSky, "Export All Targets to OneSky", "Exports translations for all cultures of all targets of this target set to OneSky.", EUserInterfaceActionType::Button, FInputChord());
+}
+
 /** Init of connection with source control server */
 void FOneSkyLocalizationServiceProvider::Init(bool bForceConnection)
 {
 	// TODO: Test if connection works?
 	bServerAvailable = true;
+
+	FOneSkyLocalizationTargetEditorCommands::Register();
 }
 
 /** API Specific close the connection with localization provider server*/
@@ -387,9 +420,14 @@ static void FileNameChanged(const FText& NewText, ETextCommit::Type CommitType,
 	FOneSkyLocalizationServiceModule::Get().AccessSettings().SetSettingsForTarget(TargetGuid, Settings->OneSkyProjectId, NewText.ToString());
 }
 
-void FOneSkyLocalizationServiceProvider::CustomizeTargetDetails(IDetailCategoryBuilder& DetailCategoryBuilder, const FGuid& TargetGuid) const
+void FOneSkyLocalizationServiceProvider::CustomizeTargetDetails(IDetailCategoryBuilder& DetailCategoryBuilder, TWeakObjectPtr<ULocalizationTarget> LocalizationTarget) const
 {
-	FOneSkyLocalizationTargetSetting* Settings = FOneSkyLocalizationServiceModule::Get().AccessSettings().GetSettingsForTarget(TargetGuid, true);
+	if (!LocalizationTarget.IsValid())
+	{
+		return;
+	}
+
+	FOneSkyLocalizationTargetSetting* Settings = FOneSkyLocalizationServiceModule::Get().AccessSettings().GetSettingsForTarget(LocalizationTarget->Settings.Guid, true);
 
 	FText ProjectText = LOCTEXT("OneSkyProjectIdLabel", "OneSky Project ID");
 	FDetailWidgetRow& ProjectRow = DetailCategoryBuilder.AddCustomRow(ProjectText);
@@ -401,7 +439,7 @@ void FOneSkyLocalizationServiceProvider::CustomizeTargetDetails(IDetailCategoryB
 	ProjectRow.ValueContent()
 		[
 			SNew(SEditableTextBox)
-			.OnTextCommitted(FOnTextCommitted::CreateStatic(&ProjectChanged, TargetGuid))
+			.OnTextCommitted(FOnTextCommitted::CreateStatic(&ProjectChanged, LocalizationTarget->Settings.Guid))
 			.Text_Lambda([Settings]
 			{
 				int32 SavedProjectId = Settings->OneSkyProjectId;
@@ -426,10 +464,57 @@ void FOneSkyLocalizationServiceProvider::CustomizeTargetDetails(IDetailCategoryB
 	FileNameRow.ValueContent()
 		[
 			SNew(SEditableTextBox)
-			.OnTextCommitted(FOnTextCommitted::CreateStatic(&FileNameChanged, TargetGuid))
+			.OnTextCommitted(FOnTextCommitted::CreateStatic(&FileNameChanged, LocalizationTarget->Settings.Guid))
 			.Text(FText::FromString(Settings->OneSkyFileName))
 		];
 }
+
+void FOneSkyLocalizationServiceProvider::CustomizeTargetToolbar(TSharedRef<FExtender>& MenuExtender, TWeakObjectPtr<ULocalizationTarget> LocalizationTarget) const
+{
+	const TSharedRef< FUICommandList > CommandList = MakeShareable(new FUICommandList);
+
+	MenuExtender->AddToolBarExtension("LocalizationService", EExtensionHook::First, CommandList,
+		FToolBarExtensionDelegate::CreateRaw(this, &FOneSkyLocalizationServiceProvider::AddTargetToolbarButtons, LocalizationTarget, CommandList));
+
+}
+
+void FOneSkyLocalizationServiceProvider::AddTargetToolbarButtons(FToolBarBuilder& ToolbarBuilder, TWeakObjectPtr<ULocalizationTarget> InLocalizationTarget, TSharedRef< FUICommandList > CommandList)
+{
+	bool bIsTargetSet = false;
+	CommandList->MapAction(FOneSkyLocalizationTargetEditorCommands::Get().ImportAllCulturesForTargetFromOneSky, FExecuteAction::CreateRaw(this, &FOneSkyLocalizationServiceProvider::ImportAllCulturesForTargetFromOneSky, InLocalizationTarget, bIsTargetSet));
+	ToolbarBuilder.AddToolBarButton(FOneSkyLocalizationTargetEditorCommands::Get().ImportAllCulturesForTargetFromOneSky, NAME_None, TAttribute<FText>(), TAttribute<FText>(), FSlateIcon(FEditorStyle::GetStyleSetName(), "TranslationEditor.ImportLatestFromLocalizationService"));
+
+	// Don't add "export all" buttons for engine targets
+	if (!InLocalizationTarget->IsMemberOfEngineTargetSet())
+	{
+		CommandList->MapAction(FOneSkyLocalizationTargetEditorCommands::Get().ExportAllCulturesForTargetToOneSky, FExecuteAction::CreateRaw(this, &FOneSkyLocalizationServiceProvider::ExportAllCulturesForTargetToOneSky, InLocalizationTarget, bIsTargetSet));
+		ToolbarBuilder.AddToolBarButton(FOneSkyLocalizationTargetEditorCommands::Get().ExportAllCulturesForTargetToOneSky, NAME_None, TAttribute<FText>(), TAttribute<FText>(), FSlateIcon(FEditorStyle::GetStyleSetName(), "TranslationEditor.ImportLatestFromLocalizationService"));
+	}
+}
+
+void FOneSkyLocalizationServiceProvider::CustomizeTargetSetToolbar(TSharedRef<FExtender>& MenuExtender, TWeakObjectPtr<ULocalizationTargetSet> InLocalizationTargetSet) const
+{
+	const TSharedRef< FUICommandList > CommandList = MakeShareable(new FUICommandList);
+
+	MenuExtender->AddToolBarExtension("LocalizationService", EExtensionHook::First, CommandList,
+		FToolBarExtensionDelegate::CreateRaw(this, &FOneSkyLocalizationServiceProvider::AddTargetSetToolbarButtons, InLocalizationTargetSet, CommandList));
+
+}
+
+void FOneSkyLocalizationServiceProvider::AddTargetSetToolbarButtons(FToolBarBuilder& ToolbarBuilder, TWeakObjectPtr<ULocalizationTargetSet> InLocalizationTargetSet, TSharedRef< FUICommandList > CommandList)
+{
+	CommandList->MapAction(FOneSkyLocalizationTargetEditorCommands::Get().ImportAllTargetsAllCulturesForTargetSetFromOneSky, FExecuteAction::CreateRaw(this, &FOneSkyLocalizationServiceProvider::ImportAllTargetsForTargetSetFromOneSky, InLocalizationTargetSet));
+	ToolbarBuilder.AddToolBarButton(FOneSkyLocalizationTargetEditorCommands::Get().ImportAllTargetsAllCulturesForTargetSetFromOneSky, NAME_None, TAttribute<FText>(), TAttribute<FText>(), FSlateIcon(FEditorStyle::GetStyleSetName(), "TranslationEditor.ImportLatestFromLocalizationService"));
+
+	// Don't add "export all" button for the engine target set
+	if (InLocalizationTargetSet.IsValid() && InLocalizationTargetSet->TargetObjects.Num() > 0 && !(InLocalizationTargetSet->TargetObjects[0]->IsMemberOfEngineTargetSet()))
+	{
+		CommandList->MapAction(FOneSkyLocalizationTargetEditorCommands::Get().ExportAllTargetsAllCulturesForTargetSetFromOneSky, FExecuteAction::CreateRaw(this, &FOneSkyLocalizationServiceProvider::ExportAllTargetsForTargetSetToOneSky, InLocalizationTargetSet));
+		ToolbarBuilder.AddToolBarButton(FOneSkyLocalizationTargetEditorCommands::Get().ExportAllTargetsAllCulturesForTargetSetFromOneSky, NAME_None, TAttribute<FText>(), TAttribute<FText>(), FSlateIcon(FEditorStyle::GetStyleSetName(), "TranslationEditor.ImportLatestFromLocalizationService"));
+	}
+}
+
+
 #endif //LOCALIZATION_SERVICES_WITH_SLATE
 
 TSharedPtr<IOneSkyLocalizationServiceWorker, ESPMode::ThreadSafe> FOneSkyLocalizationServiceProvider::CreateWorker(const FName& InOperationName) const
@@ -526,4 +611,359 @@ ELocalizationServiceOperationCommandResult::Type FOneSkyLocalizationServiceProvi
 }
 
 
+void FOneSkyLocalizationServiceProvider::ImportAllCulturesForTargetFromOneSky(TWeakObjectPtr<ULocalizationTarget> LocalizationTarget, bool bIsTargetSet)
+{
+	check(LocalizationTarget.IsValid());
+
+	if (!bIsTargetSet)
+	{
+		GWarn->BeginSlowTask(LOCTEXT("ImportingFromLocalizationService", "Importing Latest from Localization Service..."), true);
+	}
+
+	FString EngineOrGamePath = LocalizationTarget->IsMemberOfEngineTargetSet() ? "Engine" : "Game";
+
+	for (FCultureStatistics CultureStats : LocalizationTarget->Settings.SupportedCulturesStatistics)
+	{
+		ILocalizationServiceProvider& Provider = ILocalizationServiceModule::Get().GetProvider();
+		TSharedRef<FDownloadLocalizationTargetFile, ESPMode::ThreadSafe> DownloadTargetFileOp = ILocalizationServiceOperation::Create<FDownloadLocalizationTargetFile>();
+		DownloadTargetFileOp->SetInTargetGuid(LocalizationTarget->Settings.Guid);
+		DownloadTargetFileOp->SetInLocale(CultureStats.CultureName);
+
+		// Put the intermediary .po files in a temporary directory in Saved for now
+		FString Path = FPaths::GameSavedDir() / "Temp" / EngineOrGamePath / LocalizationTarget->Settings.Name / CultureStats.CultureName / LocalizationTarget->Settings.Name + ".po";
+		FPaths::MakePathRelativeTo(Path, *FPaths::GameDir());
+		DownloadTargetFileOp->SetInRelativeOutputFilePathAndName(Path);
+
+		FilesDownloadingForImportFromOneSky.Add(Path);
+		IPlatformFile &PlatformFile = FPlatformFileManager::Get().GetPlatformFile();
+		// Delete this file if it exists so we don't accidentally import old data
+		PlatformFile.DeleteFile(*Path);
+
+		auto OperationCompleteDelegate = FLocalizationServiceOperationComplete::CreateRaw(this, &FOneSkyLocalizationServiceProvider::ImportCultureForTargetFromOneSky_Callback, bIsTargetSet);
+
+		Provider.Execute(DownloadTargetFileOp, TArray<FLocalizationServiceTranslationIdentifier>(), ELocalizationServiceOperationConcurrency::Asynchronous, OperationCompleteDelegate);
+	}
+}
+
+
+
+void FOneSkyLocalizationServiceProvider::ImportCultureForTargetFromOneSky_Callback(const FLocalizationServiceOperationRef& Operation, ELocalizationServiceOperationCommandResult::Type Result, bool bIsTargetSet)
+{
+	TSharedPtr<FDownloadLocalizationTargetFile, ESPMode::ThreadSafe> DownloadLocalizationTargetOp = StaticCastSharedRef<FDownloadLocalizationTargetFile>(Operation);
+	bool bError = !(Result == ELocalizationServiceOperationCommandResult::Succeeded);
+	FText ErrorText = FText::GetEmpty();
+	FGuid InTargetGuid;
+	FString InLocale;
+	FString InRelativeOutputFilePathAndName;
+	FString AbsoluteFilePathAndName;
+	FString TargetName;
+	ULocalizationTarget* Target = nullptr;
+	bool bFinishedDownloading = false;
+	if (DownloadLocalizationTargetOp.IsValid())
+	{
+		ErrorText = DownloadLocalizationTargetOp->GetOutErrorText();
+
+		InTargetGuid = DownloadLocalizationTargetOp->GetInTargetGuid();
+		InLocale = DownloadLocalizationTargetOp->GetInLocale();
+		InRelativeOutputFilePathAndName = DownloadLocalizationTargetOp->GetInRelativeOutputFilePathAndName();
+
+		TargetName = FPaths::GetBaseFilename(InRelativeOutputFilePathAndName);
+		FString EngineOrGamePath = FPaths::GetBaseFilename(FPaths::GetPath(FPaths::GetPath(FPaths::GetPath(InRelativeOutputFilePathAndName))));
+		bool bIsEngineTarget = EngineOrGamePath == "Engine";
+		Target = FLocalizationModule::Get().GetLocalizationTargetByName(TargetName, bIsEngineTarget);
+
+		// Remove each file we get a callback for so we know when we've gotten a callback for all of them
+		FilesDownloadingForImportFromOneSky.Remove(InRelativeOutputFilePathAndName);
+
+		int32 TotalNumber = 0;
+		if (bIsTargetSet)
+		{
+			for (ULocalizationTarget* LocalizationTarget : Target->GetOuterULocalizationTargetSet()->TargetObjects)
+			{
+				TotalNumber += LocalizationTarget->Settings.SupportedCulturesStatistics.Num();
+			}
+		}
+		else
+		{
+			TotalNumber = Target->Settings.SupportedCulturesStatistics.Num();
+		}
+
+		// Update progress bar
+		GWarn->StatusUpdate(TotalNumber - FilesDownloadingForImportFromOneSky.Num(),
+			TotalNumber,
+			LOCTEXT("DownloadingFilesFromLocalizationService", "Downloading Files from Localization Service..."));
+
+		AbsoluteFilePathAndName = FPaths::ConvertRelativePathToFull(FPaths::GameDir() / InRelativeOutputFilePathAndName);
+
+		// Once we have gotten the callback for each file for this import, then start importing
+		if (FilesDownloadingForImportFromOneSky.Num() == 0)
+		{
+			GWarn->StatusUpdate(100, 100, LOCTEXT("ImportFromLocalizationServiceFinished", "Import from Localization Service Complete!"));
+			GWarn->EndSlowTask();
+			bFinishedDownloading = true;
+		}
+	}
+	if (!bError && ErrorText.IsEmpty())
+	{
+		if (!DownloadLocalizationTargetOp.IsValid())
+		{
+			bError = true;
+		}
+
+		if ( !InRelativeOutputFilePathAndName.IsEmpty())
+		{
+			if (!FPaths::FileExists(AbsoluteFilePathAndName))
+			{
+				bError = true;
+			}
+		}
+		else
+		{
+			bError = true;
+		}
+
+		if (bError)
+		{
+			if (ErrorText.IsEmpty())
+			{
+				ErrorText = LOCTEXT("DownloadLatestFromLocalizationServiceFileProcessError", "An error occured when processing the file downloaded from the Localization Service.");
+			}
+		}
+	}
+	else
+	{
+		bError = true;
+		if (ErrorText.IsEmpty())
+		{
+			ErrorText = LOCTEXT("DownloadLatestFromLocalizationServiceDownloadError", "An error occured while downloading the file from the Localization Service.");
+		}
+	}
+
+	if (bError || !ErrorText.IsEmpty())
+	{
+		if (ErrorText.IsEmpty())
+		{
+			ErrorText = LOCTEXT("DownloadLatestFromLocalizationServiceUnspecifiedError", "An unspecified error occured when trying download and import from the Localization Service.");
+		}
+
+		FText ErrorNotify = FText::Format(LOCTEXT("ImportLatestForAllCulturesForTargetFromOneSkyFail", "{0} translations for {1} target failed to import from OneSky!"), FText::FromString(InLocale), FText::FromString(TargetName));
+		FMessageLog OneSkyMessageLog("OneSky");
+		OneSkyMessageLog.Error(ErrorNotify);
+		OneSkyMessageLog.Error(ErrorText);
+		OneSkyMessageLog.Notify(ErrorNotify);
+	}
+
+	if (bFinishedDownloading)
+	{
+		IMainFrameModule& MainFrameModule = FModuleManager::LoadModuleChecked<IMainFrameModule>(TEXT("MainFrame"));
+		const TSharedPtr<SWindow>& MainFrameParentWindow = MainFrameModule.GetParentWindow();
+		if (bIsTargetSet)
+		{
+			ULocalizationTargetSet* TargetSet = Target->GetOuterULocalizationTargetSet();
+			LocalizationCommandletTasks::ImportTargets(MainFrameParentWindow.ToSharedRef(), TargetSet->TargetObjects, FPaths::GetPath(FPaths::GetPath(FPaths::GetPath(AbsoluteFilePathAndName))));
+		}
+		else
+		{
+			LocalizationCommandletTasks::ImportTarget(MainFrameParentWindow.ToSharedRef(), Target, FPaths::GetPath(FPaths::GetPath(AbsoluteFilePathAndName)));
+		}
+	}
+}
+
+void FOneSkyLocalizationServiceProvider::ExportAllCulturesForTargetToOneSky(TWeakObjectPtr<ULocalizationTarget> LocalizationTarget, bool bIsTargetSet)
+{
+	check(LocalizationTarget.IsValid());
+
+	// If this is only one target, not a whole set, get confirmation and do export here (otherwise this is handled in the calling function)
+	if (!bIsTargetSet)
+	{
+		bool bAccepted = GWarn->YesNof(FText::Format(
+			LOCTEXT("ExportAllCulturesForTargetToOneSkyConfirm", "All data in OneSky for target {0} will be overwritten with your local copy!\nThis cannot be undone.\nAre you sure you want to export all cultures for this target to OneSky?"),
+			FText::FromString(LocalizationTarget->Settings.Name)));
+		
+		if (!bAccepted)
+		{
+			return;
+		}
+
+		IMainFrameModule& MainFrameModule = FModuleManager::LoadModuleChecked<IMainFrameModule>(TEXT("MainFrame"));
+		const TSharedPtr<SWindow>& MainFrameParentWindow = MainFrameModule.GetParentWindow();
+		FString EngineOrGamePath = LocalizationTarget->IsMemberOfEngineTargetSet() ? "Engine" : "Game";
+		FString AbsoluteFolderPath = FPaths::ConvertRelativePathToFull(FPaths::GameSavedDir() / "Temp" / EngineOrGamePath / LocalizationTarget->Settings.Name);
+		
+		// Delete old files if they exists so we don't accidentally export old data
+		IPlatformFile &PlatformFile = FPlatformFileManager::Get().GetPlatformFile();
+		PlatformFile.DeleteDirectoryRecursively(*AbsoluteFolderPath);
+
+		// Export to file
+		LocalizationCommandletTasks::ExportTarget(MainFrameParentWindow.ToSharedRef(), LocalizationTarget.Get(), AbsoluteFolderPath);
+
+		GWarn->BeginSlowTask(LOCTEXT("ExportingToOneSky", "Exporting Latest to OneSky..."), true);
+	}
+
+	for (FCultureStatistics CultureStats : LocalizationTarget->Settings.SupportedCulturesStatistics)
+	{
+		ILocalizationServiceProvider& Provider = ILocalizationServiceModule::Get().GetProvider();
+		TSharedRef<FUploadLocalizationTargetFile, ESPMode::ThreadSafe> UploadFileOp = ILocalizationServiceOperation::Create<FUploadLocalizationTargetFile>();
+		UploadFileOp->SetInTargetGuid(LocalizationTarget->Settings.Guid);
+		UploadFileOp->SetInLocale(CultureStats.CultureName);
+		FString EngineOrGamePath = LocalizationTarget->IsMemberOfEngineTargetSet() ? "Engine" : "Game";
+
+		// Put the intermediary .po files in a temporary directory in Saved for now
+		FString Path = FPaths::GameSavedDir() / "Temp" / EngineOrGamePath / LocalizationTarget->Settings.Name / CultureStats.CultureName / LocalizationTarget->Settings.Name + ".po";
+		FPaths::MakePathRelativeTo(Path, *FPaths::GameDir());
+		UploadFileOp->SetInRelativeInputFilePathAndName(Path);
+
+		FilesUploadingForExportToOneSky.Add(Path);
+
+		Provider.Execute(UploadFileOp, TArray<FLocalizationServiceTranslationIdentifier>(), ELocalizationServiceOperationConcurrency::Asynchronous, 
+			FLocalizationServiceOperationComplete::CreateRaw(this, &FOneSkyLocalizationServiceProvider::ExportCultureForTargetToOneSky_Callback, bIsTargetSet));
+	}
+}
+
+
+
+void FOneSkyLocalizationServiceProvider::ExportCultureForTargetToOneSky_Callback(const FLocalizationServiceOperationRef& Operation, ELocalizationServiceOperationCommandResult::Type Result, bool bIsTargetSet)
+{
+	TSharedPtr<FUploadLocalizationTargetFile, ESPMode::ThreadSafe> UploadLocalizationTargetOp = StaticCastSharedRef<FUploadLocalizationTargetFile>(Operation);
+	bool bError = !(Result == ELocalizationServiceOperationCommandResult::Succeeded);
+	FText ErrorText = FText::GetEmpty();
+	FGuid InTargetGuid;
+	FString InRelativeInputFilePathAndName;
+	FString TargetName = "";
+	FString CultureName = "";
+	ULocalizationTarget* Target = nullptr;
+	if (UploadLocalizationTargetOp.IsValid())
+	{
+		InTargetGuid = UploadLocalizationTargetOp->GetInTargetGuid();
+		CultureName = UploadLocalizationTargetOp->GetInLocale();
+		InRelativeInputFilePathAndName = UploadLocalizationTargetOp->GetInRelativeInputFilePathAndName();
+
+		TargetName = FPaths::GetBaseFilename(InRelativeInputFilePathAndName);
+		FString EngineOrGamePath = FPaths::GetBaseFilename(FPaths::GetPath(FPaths::GetPath(FPaths::GetPath(InRelativeInputFilePathAndName))));
+		bool bIsEngineTarget = EngineOrGamePath == "Engine";
+		Target = FLocalizationModule::Get().GetLocalizationTargetByName(TargetName, bIsEngineTarget);
+
+		// Remove each file we get a callback for so we know when we've gotten a callback for all of them
+		FilesDownloadingForImportFromOneSky.Remove(InRelativeInputFilePathAndName);
+
+		ErrorText = UploadLocalizationTargetOp->GetOutErrorText();
+		InTargetGuid = UploadLocalizationTargetOp->GetInTargetGuid();
+
+		Target = FLocalizationModule::Get().GetLocalizationTargetByName(TargetName, bIsEngineTarget);
+
+		FilesUploadingForExportToOneSky.Remove(InRelativeInputFilePathAndName);
+
+		int32 TotalNumber = 0;
+		if (bIsTargetSet)
+		{
+			for (ULocalizationTarget* LocalizationTarget : Target->GetOuterULocalizationTargetSet()->TargetObjects)
+			{
+				TotalNumber += LocalizationTarget->Settings.SupportedCulturesStatistics.Num();
+			}
+		}
+		else
+		{
+			TotalNumber = Target->Settings.SupportedCulturesStatistics.Num();
+		}
+
+		// Update progress bar
+		GWarn->StatusUpdate(TotalNumber - FilesUploadingForExportToOneSky.Num(),
+			TotalNumber,
+			LOCTEXT("UploadingFilestoLocalizationService", "Uploading Files to Localization Service..."));
+
+		if (FilesUploadingForExportToOneSky.Num() == 0)
+		{
+			GWarn->EndSlowTask();
+		}
+	}
+
+	// Try to get display name
+	FInternationalization& I18N = FInternationalization::Get();
+	FCulturePtr CulturePtr = I18N.GetCulture(CultureName);
+	FString CultureDisplayName = CultureName;
+	if (CulturePtr.IsValid())
+	{
+		CultureName = CulturePtr->GetDisplayName();
+	}
+
+	if (!bError && ErrorText.IsEmpty())
+	{
+		FText SuccessText = FText::Format(LOCTEXT("ExportTranslationsToTranslationServiceSuccess", "{0} translations for {1} target uploaded for processing to Translation Service."), FText::FromString(CultureDisplayName), FText::FromString(TargetName));
+		FMessageLog TranslationEditorMessageLog("TranslationEditor");
+		TranslationEditorMessageLog.Info(SuccessText);
+		TranslationEditorMessageLog.Notify(SuccessText, EMessageSeverity::Info, true);
+	}
+	else
+	{
+		if (ErrorText.IsEmpty())
+		{
+			ErrorText = LOCTEXT("ExportToLocalizationServiceUnspecifiedError", "An unspecified error occured when trying to export to the Localization Service.");
+		}
+
+		FText ErrorNotify = FText::Format(LOCTEXT("SaveSelectedTranslationsToTranslationServiceFail", "{0} translations for {1} target failed to export to Translation Service!"), FText::FromString(CultureDisplayName), FText::FromString(TargetName));
+		FMessageLog TranslationEditorMessageLog("TranslationEditor");
+		TranslationEditorMessageLog.Error(ErrorNotify);
+		TranslationEditorMessageLog.Error(ErrorText);
+		TranslationEditorMessageLog.Notify(ErrorNotify);
+	}
+}
+
+void FOneSkyLocalizationServiceProvider::ImportAllTargetsForTargetSetFromOneSky(TWeakObjectPtr<ULocalizationTargetSet> LocalizationTargetSet)
+{
+	check(LocalizationTargetSet.IsValid());
+
+	GWarn->BeginSlowTask(LOCTEXT("ImportingFromLocalizationService", "Importing Latest from Localization Service..."), true);
+
+	bool bIsTargetSet = true;
+	for (ULocalizationTarget* LocalizationTarget : LocalizationTargetSet->TargetObjects)
+	{
+		ImportAllCulturesForTargetFromOneSky(LocalizationTarget, bIsTargetSet);
+	}
+}
+
+void FOneSkyLocalizationServiceProvider::ExportAllTargetsForTargetSetToOneSky(TWeakObjectPtr<ULocalizationTargetSet> LocalizationTargetSet)
+{
+	// If this is only one target, not a whole set, get confirmation and do export here
+
+	if (!LocalizationTargetSet.IsValid())
+	{
+		return;
+	}
+
+	if (LocalizationTargetSet->TargetObjects.Num() < 1)
+	{
+		return;
+	}
+
+	FString EngineOrGamePath = LocalizationTargetSet->TargetObjects[0]->IsMemberOfEngineTargetSet() ? "Engine" : "Game";
+
+	bool bAccepted = GWarn->YesNof(
+		FText::Format(LOCTEXT("ExportAllCulturesForTargetToOneSkyConfirm", "All data in OneSky for the {0} set of Targets will be overwritten with your local copy!\nThis cannot be undone.\nAre you sure you want to export all cultures for all targets for this set of targets to OneSky?"),
+		FText::FromString(EngineOrGamePath)));
+
+	if (!bAccepted)
+	{
+		return;
+	}
+
+	IMainFrameModule& MainFrameModule = FModuleManager::LoadModuleChecked<IMainFrameModule>(TEXT("MainFrame"));
+	const TSharedPtr<SWindow>& MainFrameParentWindow = MainFrameModule.GetParentWindow();
+	FString AbsoluteFolderPath = FPaths::ConvertRelativePathToFull(FPaths::GameSavedDir() / "Temp" / EngineOrGamePath / "");
+
+	IPlatformFile &PlatformFile = FPlatformFileManager::Get().GetPlatformFile();
+	// Delete old files if they exists so we don't accidentally export old data
+	PlatformFile.DeleteDirectoryRecursively(*AbsoluteFolderPath);
+
+	LocalizationCommandletTasks::ExportTargets(MainFrameParentWindow.ToSharedRef(), LocalizationTargetSet->TargetObjects, AbsoluteFolderPath);
+
+	GWarn->BeginSlowTask(LOCTEXT("ExportingToOneSky", "Exporting Latest to OneSky..."), true);
+
+	bool bIsTargetSet = true;
+	for (ULocalizationTarget* LocalizationTarget : LocalizationTargetSet->TargetObjects)
+	{
+		ExportAllCulturesForTargetToOneSky(LocalizationTarget, bIsTargetSet);
+	}
+}
+
+
 #undef LOCTEXT_NAMESPACE
diff --git a/Engine/Plugins/Developer/OneSkyLocalizationService/Source/OneSkyLocalizationService/Private/OneSkyLocalizationServiceProvider.h b/Engine/Plugins/Developer/OneSkyLocalizationService/Source/OneSkyLocalizationService/Private/OneSkyLocalizationServiceProvider.h
--- a/Engine/Plugins/Developer/OneSkyLocalizationService/Source/OneSkyLocalizationService/Private/OneSkyLocalizationServiceProvider.h
+++ b/Engine/Plugins/Developer/OneSkyLocalizationService/Source/OneSkyLocalizationService/Private/OneSkyLocalizationServiceProvider.h
@@ -80,7 +80,9 @@ public:
 	virtual void Tick() override;
 #if LOCALIZATION_SERVICES_WITH_SLATE
 	virtual void CustomizeSettingsDetails(IDetailCategoryBuilder& DetailCategoryBuilder) const override;
-	virtual void CustomizeTargetDetails(IDetailCategoryBuilder& DetailCategoryBuilder, const FGuid& TargetGuid) const override;
+	virtual void CustomizeTargetDetails(IDetailCategoryBuilder& DetailCategoryBuilder, TWeakObjectPtr<ULocalizationTarget> LocalizationTarget) const override;
+	virtual void CustomizeTargetToolbar(TSharedRef<FExtender>& MenuExtender, TWeakObjectPtr<ULocalizationTarget> LocalizationTarget) const override;
+	virtual void CustomizeTargetSetToolbar(TSharedRef<FExtender>& MenuExtender, TWeakObjectPtr<ULocalizationTargetSet> ULocalizationTargetSet) const override;
 #endif // LOCALIZATION_SERVICES_WITH_SLATE
 
 	/**
@@ -133,6 +135,30 @@ private:
 	 */
 	ELocalizationServiceOperationCommandResult::Type IssueCommand(class FOneSkyLocalizationServiceCommand& InCommand, const bool bSynchronous);
 
+	/** Add the buttons to the toolbar for this localization target */
+	void AddTargetToolbarButtons(FToolBarBuilder& ToolbarBuilder, TWeakObjectPtr<ULocalizationTarget> InLocalizationTarget, TSharedRef< FUICommandList > CommandList);
+
+	/** Add the buttons to the toolbar for this set of localization targets */
+	void AddTargetSetToolbarButtons(FToolBarBuilder& ToolbarBuilder, TWeakObjectPtr<ULocalizationTargetSet> InLocalizationTargetSet, TSharedRef< FUICommandList > CommandList);
+
+	/** Download and import all translations for all cultures for the specified target from OneSky */
+	void ImportAllCulturesForTargetFromOneSky(TWeakObjectPtr<ULocalizationTarget> LocalizationTarget, bool bIsTargetSet);
+
+	/** Called when done downloading localization data for culture for a target from OneSky */
+	void ImportCultureForTargetFromOneSky_Callback(const FLocalizationServiceOperationRef& Operation, ELocalizationServiceOperationCommandResult::Type Result, bool yes);
+
+	/** Export and upload all cultures for a localization target to OneSky */
+	void ExportAllCulturesForTargetToOneSky(TWeakObjectPtr<ULocalizationTarget> LocalizationTarget, bool bIsTargetSet);
+
+	/** Called when done uploading localization data for a culture for a target from OnSky*/
+	void ExportCultureForTargetToOneSky_Callback(const FLocalizationServiceOperationRef& Operation, ELocalizationServiceOperationCommandResult::Type Result, bool bIsTargetSet);
+
+	/** Download and import all translations for all cultures for all targets for the specified target set from OneSky */
+	void ImportAllTargetsForTargetSetFromOneSky(TWeakObjectPtr<ULocalizationTargetSet> LocalizationTargetSet);
+
+	/** Export and upload all cultures for all targets for a localization target set to OneSky */
+	void ExportAllTargetsForTargetSetToOneSky(TWeakObjectPtr<ULocalizationTargetSet> LocalizationTargetSet);
+
 private:
 
 	/** Indicates if Localization service integration is available or not. */
@@ -152,4 +178,10 @@ private:
 
 	/** Queue for import status tasks  */
 	TQueue < FShowImportTaskQueueItem, EQueueMode::Mpsc > ShowImportTaskQueue;
+
+	/** Array of file names being downloaded for import from OneSky */
+	TArray < FString > FilesDownloadingForImportFromOneSky;
+
+	/** Array of file names being uploaded for export into OneSky */
+	TArray < FString > FilesUploadingForExportToOneSky;
 };
\ No newline at end of file
diff --git a/Engine/Source/Developer/LocalizationService/Private/DefaultLocalizationServiceProvider.cpp b/Engine/Source/Developer/LocalizationService/Private/DefaultLocalizationServiceProvider.cpp
--- a/Engine/Source/Developer/LocalizationService/Private/DefaultLocalizationServiceProvider.cpp
+++ b/Engine/Source/Developer/LocalizationService/Private/DefaultLocalizationServiceProvider.cpp
@@ -73,10 +73,21 @@ void FDefaultLocalizationServiceProvider::CustomizeSettingsDetails(IDetailCatego
 
 }
 
-void FDefaultLocalizationServiceProvider::CustomizeTargetDetails(IDetailCategoryBuilder& DetailCategoryBuilder, const FGuid& TargetGuid) const
+void FDefaultLocalizationServiceProvider::CustomizeTargetDetails(IDetailCategoryBuilder& DetailCategoryBuilder, TWeakObjectPtr<ULocalizationTarget> LocalizationTarget) const
 {
 
 }
+
+void FDefaultLocalizationServiceProvider::CustomizeTargetToolbar(TSharedRef<FExtender>& MenuExtender, TWeakObjectPtr<ULocalizationTarget> LocalizationTarget) const
+{
+
+}
+
+void FDefaultLocalizationServiceProvider::CustomizeTargetSetToolbar(TSharedRef<FExtender>& MenuExtender, TWeakObjectPtr<ULocalizationTargetSet> LocalizationTargetSet) const
+{
+
+}
+
 #endif //LOCALIZATION_SERVICES_WITH_SLATE
 
 #undef LOCTEXT_NAMESPACE
\ No newline at end of file
diff --git a/Engine/Source/Developer/LocalizationService/Private/DefaultLocalizationServiceProvider.h b/Engine/Source/Developer/LocalizationService/Private/DefaultLocalizationServiceProvider.h
--- a/Engine/Source/Developer/LocalizationService/Private/DefaultLocalizationServiceProvider.h
+++ b/Engine/Source/Developer/LocalizationService/Private/DefaultLocalizationServiceProvider.h
@@ -25,7 +25,9 @@ public:
 	virtual void Tick() override;
 #if LOCALIZATION_SERVICES_WITH_SLATE
 	virtual void CustomizeSettingsDetails(IDetailCategoryBuilder& DetailCategoryBuilder) const override;
-	virtual void CustomizeTargetDetails(IDetailCategoryBuilder& DetailCategoryBuilder, const FGuid& TargetGuid) const override;
+	virtual void CustomizeTargetDetails(IDetailCategoryBuilder& DetailCategoryBuilder, TWeakObjectPtr<ULocalizationTarget> LocalizationTarget) const override;
+	virtual void CustomizeTargetToolbar(TSharedRef<FExtender>& MenuExtender, TWeakObjectPtr<ULocalizationTarget> LocalizationTarget) const override;
+	virtual void CustomizeTargetSetToolbar(TSharedRef<FExtender>& MenuExtender, TWeakObjectPtr<ULocalizationTargetSet> LocalizationTargetSet) const override;
 #endif // LOCALIZATION_SERVICES_WITH_SLATE
 };
 
diff --git a/Engine/Source/Developer/LocalizationService/Public/ILocalizationServiceProvider.h b/Engine/Source/Developer/LocalizationService/Public/ILocalizationServiceProvider.h
--- a/Engine/Source/Developer/LocalizationService/Public/ILocalizationServiceProvider.h
+++ b/Engine/Source/Developer/LocalizationService/Public/ILocalizationServiceProvider.h
@@ -8,6 +8,8 @@
 #include "ILocalizationServiceState.h"
 class IDetailCategoryBuilder;
 class ULocalizationTarget;
+class ULocalizationTargetSet;
+class FExtender;
 
 #define LOCALIZATION_SERVICES_WITH_SLATE	(!(PLATFORM_LINUX && IS_PROGRAM))
 
@@ -202,7 +204,20 @@ public:
 	* Create a settings widget for display in the localization target editor.
 	* @returns a widget used to edit target-specific localization service settings.
 	*/
-	virtual void CustomizeTargetDetails(IDetailCategoryBuilder& DetailCategoryBuilder, const FGuid& TargetGuid) const = 0;
+	virtual void CustomizeTargetDetails(IDetailCategoryBuilder& DetailCategoryBuilder, TWeakObjectPtr<ULocalizationTarget> LocalizationTarget) const = 0;
+
+	/**
+	* Create a settings widget for display in the localization target editor.
+	* @returns a widget used to edit target-specific localization service settings.
+	*/
+	virtual void CustomizeTargetToolbar(TSharedRef<FExtender>& MenuExtender, TWeakObjectPtr<ULocalizationTarget> LocalizationTarget) const = 0;
+
+	/**
+	* Create a settings widget for display in the localization target set editor.
+	* @returns a widget used to edit target set-specific localization service settings.
+	*/
+	virtual void CustomizeTargetSetToolbar(TSharedRef<FExtender>& MenuExtender, TWeakObjectPtr<ULocalizationTargetSet> LocalizationTargetSet) const = 0;
+
 #endif //LOCALIZATION_SERVICES_WITH_SLATE
 };
 
diff --git a/Engine/Source/Editor/LocalizationDashboard/LocalizationDashboard.Build.cs b/Engine/Source/Editor/LocalizationDashboard/LocalizationDashboard.Build.cs
--- a/Engine/Source/Editor/LocalizationDashboard/LocalizationDashboard.Build.cs
+++ b/Engine/Source/Editor/LocalizationDashboard/LocalizationDashboard.Build.cs
@@ -50,11 +50,11 @@ public class LocalizationDashboard : ModuleRules
 		);
 
         CircularlyReferencedDependentModules.AddRange(
-            new string[]
-			{
-                 "MainFrame",
-                 "TranslationEditor"
-			}
-            );
+           new string[] {
+                "LocalizationService",
+				"MainFrame",
+				"TranslationEditor"
+            }
+           );
 	}
 }
\ No newline at end of file
diff --git a/Engine/Source/Editor/LocalizationDashboard/Private/LocalizationTargetDetailCustomization.cpp b/Engine/Source/Editor/LocalizationDashboard/Private/LocalizationTargetDetailCustomization.cpp
--- a/Engine/Source/Editor/LocalizationDashboard/Private/LocalizationTargetDetailCustomization.cpp
+++ b/Engine/Source/Editor/LocalizationDashboard/Private/LocalizationTargetDetailCustomization.cpp
@@ -177,16 +177,25 @@ void FLocalizationTargetDetailCustomization::CustomizeDetails(IDetailLayoutBuild
 		NativeCultureIndexPropertyHandle = MemberPropertyHandle;
 	});
 
+	const ILocalizationServiceProvider& LSP = ILocalizationServiceModule::Get().GetProvider();
 	PropertyCustomizationMap.Add(GET_MEMBER_NAME_CHECKED(FLocalizationTargetSettings, SupportedCulturesStatistics), [&](const TSharedRef<IPropertyHandle>& MemberPropertyHandle, IDetailCategoryBuilder& DetailCategoryBuilder)
 	{
+
 		SupportedCulturesStatisticsPropertyHandle = MemberPropertyHandle;
 
 		SupportedCulturesStatisticsPropertyHandle_OnNumElementsChanged = FSimpleDelegate::CreateSP(this, &FLocalizationTargetDetailCustomization::RebuildListedCulturesList);
 		SupportedCulturesStatisticsPropertyHandle->AsArray()->SetOnNumElementsChanged(SupportedCulturesStatisticsPropertyHandle_OnNumElementsChanged);
 
 		FLocalizationTargetEditorCommands::Register();
+		auto Commands = FLocalizationTargetEditorCommands::Get();
 		const TSharedRef< FUICommandList > CommandList = MakeShareable(new FUICommandList);
-		FToolBarBuilder ToolBarBuilder( CommandList, FMultiBoxCustomization::AllowCustomization("LocalizationTargetEditor") );
+		// Let the localization service extend this toolbar
+		TSharedRef<FExtender> LocalizationServiceExtender = MakeShareable(new FExtender);
+		if (LocalizationTarget.IsValid() && ILocalizationServiceModule::Get().IsEnabled())
+		{
+			LSP.CustomizeTargetToolbar(LocalizationServiceExtender, LocalizationTarget);
+		}
+		FToolBarBuilder ToolBarBuilder(CommandList, FMultiBoxCustomization::AllowCustomization("LocalizationTargetEditor"), LocalizationServiceExtender);
 
 		TAttribute<FText> GatherToolTipTextAttribute = TAttribute<FText>::Create(TAttribute<FText>::FGetter::CreateLambda([this]() -> FText
 			{
@@ -206,6 +215,12 @@ void FLocalizationTargetDetailCustomization::CustomizeDetails(IDetailLayoutBuild
 
 		CommandList->MapAction( FLocalizationTargetEditorCommands::Get().CompileAllCultures, FExecuteAction::CreateSP(this, &FLocalizationTargetDetailCustomization::CompileAllCultures), FCanExecuteAction::CreateSP(this, &FLocalizationTargetDetailCustomization::CanCompileAllCultures));
 		ToolBarBuilder.AddToolBarButton(FLocalizationTargetEditorCommands::Get().CompileAllCultures, NAME_None, TAttribute<FText>(), TAttribute<FText>(), FSlateIcon(FEditorStyle::GetStyleSetName(), "LocalizationTargetEditor.CompileAllCultures"));
+		
+		if (ILocalizationServiceModule::Get().IsEnabled())
+		{
+			ToolBarBuilder.BeginSection("LocalizationService");
+			ToolBarBuilder.EndSection();
+		}
 
 		BuildListedCulturesList();
 
@@ -362,11 +377,9 @@ void FLocalizationTargetDetailCustomization::CustomizeDetails(IDetailLayoutBuild
 	{
 		// Assign this to the "important" category to make it show up just under the "target" category
 		IDetailCategoryBuilder& ServiceProviderCategoryBuilder = DetailBuilder.EditCategory(FName("LocalizationServiceProvider"), LOCTEXT("LocalizationServiceProviderCategoryLabel","Localization Service Provider"), ECategoryPriority::Important);
-
-		const ILocalizationServiceProvider& LSP = ILocalizationServiceModule::Get().GetProvider();
 		if (LocalizationTarget.IsValid())
 		{
-			LSP.CustomizeTargetDetails(ServiceProviderCategoryBuilder, LocalizationTarget->Settings.Guid);
+			LSP.CustomizeTargetDetails(ServiceProviderCategoryBuilder, LocalizationTarget);
 		}
 	}
 }
diff --git a/Engine/Source/Editor/LocalizationDashboard/Private/LocalizationTargetSetDetailCustomization.cpp b/Engine/Source/Editor/LocalizationDashboard/Private/LocalizationTargetSetDetailCustomization.cpp
--- a/Engine/Source/Editor/LocalizationDashboard/Private/LocalizationTargetSetDetailCustomization.cpp
+++ b/Engine/Source/Editor/LocalizationDashboard/Private/LocalizationTargetSetDetailCustomization.cpp
@@ -86,6 +86,7 @@ void FLocalizationTargetSetDetailCustomization::CustomizeDetails(IDetailLayoutBu
 	}
 
 	{
+		const ILocalizationServiceProvider& LSP = ILocalizationServiceModule::Get().GetProvider();
 		TargetObjectsPropertyHandle = DetailBuilder.GetProperty(GET_MEMBER_NAME_CHECKED(ULocalizationTargetSet,TargetObjects));
 		if (TargetObjectsPropertyHandle->IsValidHandle())
 		{
@@ -98,7 +99,15 @@ void FLocalizationTargetSetDetailCustomization::CustomizeDetails(IDetailLayoutBu
 
 			FLocalizationDashboardCommands::Register();
 			const TSharedRef< FUICommandList > CommandList = MakeShareable(new FUICommandList);
-			FToolBarBuilder ToolBarBuilder( CommandList, FMultiBoxCustomization::AllowCustomization("LocalizationDashboard") );
+
+			// Let the localization service extend this toolbar
+			TSharedRef<FExtender> LocalizationServiceExtender = MakeShareable(new FExtender);
+			if (TargetSet.IsValid() && ILocalizationServiceModule::Get().IsEnabled())
+			{
+				LSP.CustomizeTargetSetToolbar(LocalizationServiceExtender, TargetSet);
+			}
+
+			FToolBarBuilder ToolBarBuilder( CommandList, FMultiBoxCustomization::AllowCustomization("LocalizationDashboard"), LocalizationServiceExtender );
 
 			TAttribute<FText> GatherAllTargetsToolTipTextAttribute = TAttribute<FText>::Create(TAttribute<FText>::FGetter::CreateLambda([this]() -> FText
 				{
@@ -119,6 +128,12 @@ void FLocalizationTargetSetDetailCustomization::CustomizeDetails(IDetailLayoutBu
 			CommandList->MapAction( FLocalizationDashboardCommands::Get().CompileAllTargets, FExecuteAction::CreateSP(this, &FLocalizationTargetSetDetailCustomization::CompileAllTargets), FCanExecuteAction::CreateSP(this, &FLocalizationTargetSetDetailCustomization::CanCompileAllTargets));
 			ToolBarBuilder.AddToolBarButton(FLocalizationDashboardCommands::Get().CompileAllTargets, NAME_None, TAttribute<FText>(), TAttribute<FText>(), FSlateIcon(FEditorStyle::GetStyleSetName(), "LocalizationDashboard.CompileAllTargetsAllCultures"));
 
+			if (ILocalizationServiceModule::Get().IsEnabled())
+			{
+				ToolBarBuilder.BeginSection("LocalizationService");
+				ToolBarBuilder.EndSection();
+			}
+
 			BuildTargetsList();
 
 			DetailCategoryBuilder.AddCustomRow(TargetObjectsPropertyHandle->GetPropertyDisplayName())
