#UE4 Fixed localization plugin to use FPlatformHTTP now that has moved

diff --git a/Engine/Plugins/Developer/OneSkyLocalizationService/Source/OneSkyLocalizationService/Private/OneSkyLocalizationServiceOperations.cpp b/Engine/Plugins/Developer/OneSkyLocalizationService/Source/OneSkyLocalizationService/Private/OneSkyLocalizationServiceOperations.cpp
--- a/Engine/Plugins/Developer/OneSkyLocalizationService/Source/OneSkyLocalizationService/Private/OneSkyLocalizationServiceOperations.cpp
+++ b/Engine/Plugins/Developer/OneSkyLocalizationService/Source/OneSkyLocalizationService/Private/OneSkyLocalizationServiceOperations.cpp
@@ -257,8 +257,8 @@ bool FOneSkyCreateProjectGroupWorker::Execute(class FOneSkyLocalizationServiceCo
 	if (CreateProjectGroupOp.IsValid())
 	{
 		//we need to url encode for special characters (especially other languages)
-		InProjectGroupName = FGenericPlatformHttp::UrlEncode(CreateProjectGroupOp->GetInProjectGroupName());
-		InProjectGroupBaseLocale = FGenericPlatformHttp::UrlEncode(CreateProjectGroupOp->GetInProjectGroupBaseLocale());
+		InProjectGroupName = FPlatformHttp::UrlEncode(CreateProjectGroupOp->GetInProjectGroupName());
+		InProjectGroupBaseLocale = FPlatformHttp::UrlEncode(CreateProjectGroupOp->GetInProjectGroupBaseLocale());
 	}
 
 	FString Url = "https://platform.api.onesky.io/1/project-groups";
@@ -517,9 +517,9 @@ bool FOneSkyCreateProjectWorker::Execute(class FOneSkyLocalizationServiceCommand
 	if (CreateProjectOp.IsValid())
 	{
 		//we need to url encode for special characters (especially other languages)
-		InProjectGroupName = FGenericPlatformHttp::UrlEncode(CreateProjectOp->GetInProjectGroupName());
-		InProjectDescription = FGenericPlatformHttp::UrlEncode(CreateProjectOp->GetInProjectDescription());
-		InProjectType = FGenericPlatformHttp::UrlEncode(CreateProjectOp->GetInProjectType());
+		InProjectGroupName = FPlatformHttp::UrlEncode(CreateProjectOp->GetInProjectGroupName());
+		InProjectDescription = FPlatformHttp::UrlEncode(CreateProjectOp->GetInProjectDescription());
+		InProjectType = FPlatformHttp::UrlEncode(CreateProjectOp->GetInProjectType());
 		InProjectGroupId = CreateProjectOp->GetInProjectGroupId();
 	}
 
@@ -641,8 +641,8 @@ bool FOneSkyTranslationStatusWorker::Execute(class FOneSkyLocalizationServiceCom
 	if (TranslationStatusOp.IsValid())
 	{
 		InProjectId = TranslationStatusOp->GetInProjectId();
-		InFileName = FGenericPlatformHttp::UrlEncode(TranslationStatusOp->GetInFileName());
-		InLocale = FGenericPlatformHttp::UrlEncode(TranslationStatusOp->GetInLocale());
+		InFileName = FPlatformHttp::UrlEncode(TranslationStatusOp->GetInFileName());
+		InLocale = FPlatformHttp::UrlEncode(TranslationStatusOp->GetInLocale());
 	}
 	
 	FString Url = AddAuthenticationParameters(InCommand.ConnectionInfo, "https://platform.api.onesky.io/1/projects/" + FString::FromInt(InProjectId) / "translations/status");
@@ -701,7 +701,7 @@ bool FOneSkyTranslationExportWorker::Execute(class FOneSkyLocalizationServiceCom
 	if (DownloadLocTargetOp.IsValid())
 	{
 		InTargetGuid = DownloadLocTargetOp->GetInTargetGuid();
-		InLocale = FGenericPlatformHttp::UrlEncode(DownloadLocTargetOp->GetInLocale());
+		InLocale = FPlatformHttp::UrlEncode(DownloadLocTargetOp->GetInLocale());
 		InRelativeOutputFilePathAndName = DownloadLocTargetOp->GetInRelativeOutputFilePathAndName();
 	}
 
@@ -715,8 +715,8 @@ bool FOneSkyTranslationExportWorker::Execute(class FOneSkyLocalizationServiceCom
 	{
 		ILocalizationServiceProvider& Provider = ILocalizationServiceModule::Get().GetProvider();
 		InProjectId = Settings->OneSkyProjectId;
-		InSourceFileName = FGenericPlatformHttp::UrlEncode(Settings->OneSkyFileName);
-		InExportFileName = FGenericPlatformHttp::UrlEncode(FPaths::ConvertRelativePathToFull(InRelativeOutputFilePathAndName));
+		InSourceFileName = FPlatformHttp::UrlEncode(Settings->OneSkyFileName);
+		InExportFileName = FPlatformHttp::UrlEncode(FPaths::ConvertRelativePathToFull(InRelativeOutputFilePathAndName));
 	}
 
 	FString Url = AddAuthenticationParameters(InCommand.ConnectionInfo, "https://platform.api.onesky.io/1/projects/" + FString::FromInt(InProjectId) / "translations");
@@ -952,7 +952,7 @@ bool FOneSkyUploadFileWorker::Execute(class FOneSkyLocalizationServiceCommand& I
 	if (Settings != nullptr)
 	{
 		InProjectId = Settings->OneSkyProjectId;
-		InOneSkyTargetFileName = FGenericPlatformHttp::UrlEncode(Settings->OneSkyFileName);
+		InOneSkyTargetFileName = FPlatformHttp::UrlEncode(Settings->OneSkyFileName);
 	}
 
 	FString Url = AddAuthenticationParameters(InCommand.ConnectionInfo, "https://platform.api.onesky.io/1/projects/" + FString::FromInt(InProjectId) / "files");
