Plugins: Better names and categories

diff --git a/Engine/Plugins/Developer/OneSkyLocalizationService/OneSkyLocalizationService.uplugin b/Engine/Plugins/Developer/OneSkyLocalizationService/OneSkyLocalizationService.uplugin
--- a/Engine/Plugins/Developer/OneSkyLocalizationService/OneSkyLocalizationService.uplugin
+++ b/Engine/Plugins/Developer/OneSkyLocalizationService/OneSkyLocalizationService.uplugin
@@ -8,7 +8,7 @@
 	"CreatedByURL" : "http://epicgames.com",
 	"EngineVersion" : "4.8.0",
 	"Description" : "OneSky localization service",
-	"Category" : "Editor.Localization Service",
+	"Category" : "Localization",
 	"EnabledByDefault" : false,
     "IsBetaVersion": true,
 	
diff --git a/Engine/Plugins/Editor/PluginBrowser/PluginBrowser.uplugin b/Engine/Plugins/Editor/PluginBrowser/PluginBrowser.uplugin
--- a/Engine/Plugins/Editor/PluginBrowser/PluginBrowser.uplugin
+++ b/Engine/Plugins/Editor/PluginBrowser/PluginBrowser.uplugin
@@ -8,7 +8,7 @@
 	"CreatedByURL" : "http://epicgames.com",
 	"EngineVersion" : "4.2.0",
 	"Description" : "User interface for managing installed plugins and creating new ones.",
-	"Category" : "Misc",
+	"Category" : "Editor",
 	"EnabledByDefault" : true,
 	
 	"Modules" :
diff --git a/Engine/Plugins/Experimental/StereoPanorama/StereoPanorama.uplugin b/Engine/Plugins/Experimental/StereoPanorama/StereoPanorama.uplugin
--- a/Engine/Plugins/Experimental/StereoPanorama/StereoPanorama.uplugin
+++ b/Engine/Plugins/Experimental/StereoPanorama/StereoPanorama.uplugin
@@ -8,7 +8,7 @@
 	"CreatedByURL" : "http://kiteandlightning.la/",
 	"EngineVersion" : 2328018,
 	"Description" : "A plugin to capture a sequence of panoramic images",
-	"Category" : "Movie",
+	"Category" : "Recording",
     "CanContainContent" : "true",
 	
     "Modules" :
