Some fixes after CL# 22808749.

diff --git a/Engine/Plugins/Developer/OneSkyLocalizationService/Source/OneSkyLocalizationService/Private/OneSkyLocalizationServiceProvider.h b/Engine/Plugins/Developer/OneSkyLocalizationService/Source/OneSkyLocalizationService/Private/OneSkyLocalizationServiceProvider.h
--- a/Engine/Plugins/Developer/OneSkyLocalizationService/Source/OneSkyLocalizationService/Private/OneSkyLocalizationServiceProvider.h
+++ b/Engine/Plugins/Developer/OneSkyLocalizationService/Source/OneSkyLocalizationService/Private/OneSkyLocalizationServiceProvider.h
@@ -30,7 +30,7 @@ struct FLocalizationServiceTranslationIdentifierKeyFuncs : BaseKeyFuncs < TPair<
 		FLocalizationServiceTranslationIdentifier
 	> Super;
 
-	static FORCEINLINE const FLocalizationServiceTranslationIdentifier& GetSetKey(const FLocalizationServiceTranslationIdentifier& Element)
+	static FORCEINLINE const FLocalizationServiceTranslationIdentifier& GetSetKey(const TPair<FLocalizationServiceTranslationIdentifier, ValueType>& Element)
 	{
 		return Element.Key;
 	}
diff --git a/Engine/Plugins/Tests/EditorTests/Source/EditorTests/Private/UnrealEd/EditorAutomationTests.cpp b/Engine/Plugins/Tests/EditorTests/Source/EditorTests/Private/UnrealEd/EditorAutomationTests.cpp
--- a/Engine/Plugins/Tests/EditorTests/Source/EditorTests/Private/UnrealEd/EditorAutomationTests.cpp
+++ b/Engine/Plugins/Tests/EditorTests/Source/EditorTests/Private/UnrealEd/EditorAutomationTests.cpp
@@ -822,7 +822,7 @@ namespace TraceHelper
 		else
 		{
 			// Not a macro, just a plain graph; we use NULL for that.
-			VisitedTracker = GraphNodesVisited.Find(NULL);
+			VisitedTracker = GraphNodesVisited.Find(nullptr);
 			check(VisitedTracker);
 		}
 
